Docs: Removed file comments and added missing code documentation

diff --git a/Engine/Source/Editor/Matinee/Classes/MatineeTrackBoolPropHelper.h b/Engine/Source/Editor/Matinee/Classes/MatineeTrackBoolPropHelper.h
--- a/Engine/Source/Editor/Matinee/Classes/MatineeTrackBoolPropHelper.h
+++ b/Engine/Source/Editor/Matinee/Classes/MatineeTrackBoolPropHelper.h
@@ -9,12 +9,13 @@ class UMatineeTrackBoolPropHelper : public UInterpTrackHelper
 {
 	GENERATED_UCLASS_BODY()
 
+	void OnCreateTrackTextEntry(const FString& ChosenText, TWeakPtr<SWindow> Window, FString* OutputString);
 
-	// Begin UInterpTrackHelper Interface
-	virtual	bool PreCreateTrack( UInterpGroup* Group, const UInterpTrack *TrackDef, bool bDuplicatingTrack, bool bAllowPrompts ) const OVERRIDE;
-	virtual void  PostCreateTrack( UInterpTrack *Track, bool bDuplicatingTrack, int32 TrackIndex ) const OVERRIDE;
-	// End UInterpTrackHelper Interface
+public:
 
-	void OnCreateTrackTextEntry(const FString& ChosenText, TWeakPtr<SWindow> Window, FString* OutputString);
+	// UInterpTrackHelper interface
+
+	virtual	bool PreCreateTrack( UInterpGroup* Group, const UInterpTrack *TrackDef, bool bDuplicatingTrack, bool bAllowPrompts ) const override;
+	virtual void  PostCreateTrack( UInterpTrack *Track, bool bDuplicatingTrack, int32 TrackIndex ) const override;
 };
 
diff --git a/Engine/Source/Editor/Matinee/Classes/MatineeTrackColorPropHelper.h b/Engine/Source/Editor/Matinee/Classes/MatineeTrackColorPropHelper.h
--- a/Engine/Source/Editor/Matinee/Classes/MatineeTrackColorPropHelper.h
+++ b/Engine/Source/Editor/Matinee/Classes/MatineeTrackColorPropHelper.h
@@ -8,10 +8,10 @@ class UMatineeTrackColorPropHelper : public UMatineeTrackVectorPropHelper
 {
 	GENERATED_UCLASS_BODY()
 
+public:
 
-	// Begin UInterpTrackHelper Interface
-	virtual	bool PreCreateTrack( UInterpGroup* Group, const UInterpTrack *TrackDef, bool bDuplicatingTrack, bool bAllowPrompts ) const OVERRIDE;
-	virtual void  PostCreateTrack( UInterpTrack *Track, bool bDuplicatingTrack, int32 TrackIndex ) const OVERRIDE;
-	// End UInterpTrackHelper Interface
-};
+	// UInterpTrackHelper interface
 
+	virtual	bool PreCreateTrack( UInterpGroup* Group, const UInterpTrack *TrackDef, bool bDuplicatingTrack, bool bAllowPrompts ) const override;
+	virtual void  PostCreateTrack( UInterpTrack *Track, bool bDuplicatingTrack, int32 TrackIndex ) const override;
+};
diff --git a/Engine/Source/Editor/Matinee/Classes/MatineeTrackDirectorHelper.h b/Engine/Source/Editor/Matinee/Classes/MatineeTrackDirectorHelper.h
--- a/Engine/Source/Editor/Matinee/Classes/MatineeTrackDirectorHelper.h
+++ b/Engine/Source/Editor/Matinee/Classes/MatineeTrackDirectorHelper.h
@@ -1,23 +1,21 @@
 // Copyright 1998-2014 Epic Games, Inc. All Rights Reserved.
 
-/**
- *
- */
-
 #pragma once
+
 #include "MatineeTrackDirectorHelper.generated.h"
 
+
 UCLASS()
 class UMatineeTrackDirectorHelper : public UInterpTrackHelper
 {
 	GENERATED_UCLASS_BODY()
 
+	void OnAddKeyTextEntry(const FString& ChosenText, IMatineeBase* Matinee, UInterpTrack* Track);
 
-	// Begin UInterpTrackHelper Interface
-	virtual	bool PreCreateKeyframe( UInterpTrack *Track, float KeyTime ) const OVERRIDE;
-	virtual void  PostCreateKeyframe( UInterpTrack *Track, int32 KeyIndex ) const OVERRIDE;
-	// End UInterpTrackHelper Interface
+public:
 
-	void OnAddKeyTextEntry(const FString& ChosenText, IMatineeBase* Matinee, UInterpTrack* Track);
-};
+	// UInterpTrackHelper interface
 
+	virtual	bool PreCreateKeyframe( UInterpTrack *Track, float KeyTime ) const override;
+	virtual void  PostCreateKeyframe( UInterpTrack *Track, int32 KeyIndex ) const override;
+};
diff --git a/Engine/Source/Editor/Matinee/Classes/MatineeTrackEventHelper.h b/Engine/Source/Editor/Matinee/Classes/MatineeTrackEventHelper.h
--- a/Engine/Source/Editor/Matinee/Classes/MatineeTrackEventHelper.h
+++ b/Engine/Source/Editor/Matinee/Classes/MatineeTrackEventHelper.h
@@ -1,23 +1,21 @@
 // Copyright 1998-2014 Epic Games, Inc. All Rights Reserved.
 
-/**
- *
- */
-
 #pragma once
+
 #include "MatineeTrackEventHelper.generated.h"
 
+
 UCLASS()
 class UMatineeTrackEventHelper : public UInterpTrackHelper
 {
 	GENERATED_UCLASS_BODY()
 
+	void OnAddKeyTextEntry(const FText& ChosenText, ETextCommit::Type CommitInfo, IMatineeBase* Matinee, UInterpTrack* Track);
 
-	// Begin UInterpTrackHelper Interface
-	virtual	bool PreCreateKeyframe( UInterpTrack *Track, float KeyTime ) const OVERRIDE;
-	virtual void  PostCreateKeyframe( UInterpTrack *Track, int32 KeyIndex ) const OVERRIDE;
-	// End UInterpTrackHelper Interface
+public:
 
-	void OnAddKeyTextEntry(const FText& ChosenText, ETextCommit::Type CommitInfo, IMatineeBase* Matinee, UInterpTrack* Track);
-};
+	// UInterpTrackHelper interface
 
+	virtual	bool PreCreateKeyframe( UInterpTrack *Track, float KeyTime ) const override;
+	virtual void  PostCreateKeyframe( UInterpTrack *Track, int32 KeyIndex ) const override;
+};
diff --git a/Engine/Source/Editor/Matinee/Classes/MatineeTrackFloatPropHelper.h b/Engine/Source/Editor/Matinee/Classes/MatineeTrackFloatPropHelper.h
--- a/Engine/Source/Editor/Matinee/Classes/MatineeTrackFloatPropHelper.h
+++ b/Engine/Source/Editor/Matinee/Classes/MatineeTrackFloatPropHelper.h
@@ -12,12 +12,13 @@ class UMatineeTrackFloatPropHelper : public UInterpTrackHelper
 {
 	GENERATED_UCLASS_BODY()
 
+	void OnCreateTrackTextEntry(const FString& ChosenText, TSharedRef<SWindow> Window, FString* OutputString);
 
-	// Begin UInterpTrackHelper Interface
-	virtual	bool PreCreateTrack( UInterpGroup* Group, const UInterpTrack *TrackDef, bool bDuplicatingTrack, bool bAllowPrompts ) const OVERRIDE;
-	virtual void  PostCreateTrack( UInterpTrack *Track, bool bDuplicatingTrack, int32 TrackIndex ) const OVERRIDE;
-	// End UInterpTrackHelper Interface
+public:
 
-	void OnCreateTrackTextEntry(const FString& ChosenText, TSharedRef<SWindow> Window, FString* OutputString);
+	// UInterpTrackHelper interface
+
+	virtual	bool PreCreateTrack( UInterpGroup* Group, const UInterpTrack *TrackDef, bool bDuplicatingTrack, bool bAllowPrompts ) const override;
+	virtual void  PostCreateTrack( UInterpTrack *Track, bool bDuplicatingTrack, int32 TrackIndex ) const override;
 };
 
diff --git a/Engine/Source/Editor/Matinee/Classes/MatineeTrackLinearColorPropHelper.h b/Engine/Source/Editor/Matinee/Classes/MatineeTrackLinearColorPropHelper.h
--- a/Engine/Source/Editor/Matinee/Classes/MatineeTrackLinearColorPropHelper.h
+++ b/Engine/Source/Editor/Matinee/Classes/MatineeTrackLinearColorPropHelper.h
@@ -1,18 +1,20 @@
 // Copyright 1998-2014 Epic Games, Inc. All Rights Reserved.
 
-
 #pragma once
+
 #include "MatineeTrackLinearColorPropHelper.generated.h"
 
+
 UCLASS()
 class UMatineeTrackLinearColorPropHelper : public UMatineeTrackVectorPropHelper
 {
 	GENERATED_UCLASS_BODY()
 
+public:
+
+	// UInterpTrackHelper interface
 
-	// Begin UInterpTrackHelper Interface
-	virtual	bool PreCreateTrack( UInterpGroup* Group, const UInterpTrack *TrackDef, bool bDuplicatingTrack, bool bAllowPrompts ) const OVERRIDE;
-	virtual void  PostCreateTrack( UInterpTrack *Track, bool bDuplicatingTrack, int32 TrackIndex ) const OVERRIDE;
-	// End UInterpTrackHelper Interface
+	virtual	bool PreCreateTrack( UInterpGroup* Group, const UInterpTrack *TrackDef, bool bDuplicatingTrack, bool bAllowPrompts ) const override;
+	virtual void  PostCreateTrack( UInterpTrack *Track, bool bDuplicatingTrack, int32 TrackIndex ) const override;
 };
 
diff --git a/Engine/Source/Editor/Matinee/Classes/MatineeTrackParticleReplayHelper.h b/Engine/Source/Editor/Matinee/Classes/MatineeTrackParticleReplayHelper.h
--- a/Engine/Source/Editor/Matinee/Classes/MatineeTrackParticleReplayHelper.h
+++ b/Engine/Source/Editor/Matinee/Classes/MatineeTrackParticleReplayHelper.h
@@ -1,18 +1,19 @@
 // Copyright 1998-2014 Epic Games, Inc. All Rights Reserved.
 
-
 #pragma once
+
 #include "MatineeTrackParticleReplayHelper.generated.h"
 
+
 UCLASS()
 class UMatineeTrackParticleReplayHelper : public UInterpTrackHelper
 {
 	GENERATED_UCLASS_BODY()
 
+public:
 
-	// Begin UInterpTrackHelper Interface
-	virtual	bool PreCreateKeyframe( UInterpTrack *Track, float KeyTime ) const OVERRIDE;
-	virtual void  PostCreateKeyframe( UInterpTrack *Track, int32 KeyIndex ) const OVERRIDE;
-	// End UInterpTrackHelper Interface
-};
+	// UInterpTrackHelper interface
 
+	virtual	bool PreCreateKeyframe( UInterpTrack *Track, float KeyTime ) const override;
+	virtual void  PostCreateKeyframe( UInterpTrack *Track, int32 KeyIndex ) const override;
+};
diff --git a/Engine/Source/Editor/Matinee/Classes/MatineeTrackSoundHelper.h b/Engine/Source/Editor/Matinee/Classes/MatineeTrackSoundHelper.h
--- a/Engine/Source/Editor/Matinee/Classes/MatineeTrackSoundHelper.h
+++ b/Engine/Source/Editor/Matinee/Classes/MatineeTrackSoundHelper.h
@@ -1,21 +1,20 @@
 // Copyright 1998-2014 Epic Games, Inc. All Rights Reserved.
 
-/**
- *
- */
-
 #pragma once
+
 #include "MatineeTrackSoundHelper.generated.h"
 
+
 UCLASS()
 class UMatineeTrackSoundHelper : public UInterpTrackHelper
 {
 	GENERATED_UCLASS_BODY()
 
+public:
+
+	// UInterpTrackHelper interface
 
-	// Begin UInterpTrackHelper Interface
-	virtual	bool PreCreateKeyframe( UInterpTrack *Track, float KeyTime ) const OVERRIDE;
-	virtual void  PostCreateKeyframe( UInterpTrack *Track, int32 KeyIndex ) const OVERRIDE;
-	// End UInterpTrackHelper Interface
+	virtual	bool PreCreateKeyframe( UInterpTrack *Track, float KeyTime ) const override;
+	virtual void  PostCreateKeyframe( UInterpTrack *Track, int32 KeyIndex ) const override;
 };
 
diff --git a/Engine/Source/Editor/Matinee/Classes/MatineeTrackToggleHelper.h b/Engine/Source/Editor/Matinee/Classes/MatineeTrackToggleHelper.h
--- a/Engine/Source/Editor/Matinee/Classes/MatineeTrackToggleHelper.h
+++ b/Engine/Source/Editor/Matinee/Classes/MatineeTrackToggleHelper.h
@@ -1,20 +1,21 @@
 // Copyright 1998-2014 Epic Games, Inc. All Rights Reserved.
 
-
 #pragma once
+
 #include "MatineeTrackToggleHelper.generated.h"
 
+
 UCLASS()
 class UMatineeTrackToggleHelper : public UInterpTrackHelper
 {
 	GENERATED_UCLASS_BODY()
 
+	void OnAddKeyTextEntry(const FString& ChosenText, IMatineeBase* Matinee, UInterpTrack* Track);
 
-	// Begin UInterpTrackHelper Interface
-	virtual	bool PreCreateKeyframe( UInterpTrack *Track, float KeyTime ) const OVERRIDE;
-	virtual void  PostCreateKeyframe( UInterpTrack *Track, int32 KeyIndex ) const OVERRIDE;
-	// End UInterpTrackHelper Interface
+public:
 
-	void OnAddKeyTextEntry(const FString& ChosenText, IMatineeBase* Matinee, UInterpTrack* Track);
-};
+	// UInterpTrackHelper interface
 
+	virtual	bool PreCreateKeyframe( UInterpTrack *Track, float KeyTime ) const override;
+	virtual void  PostCreateKeyframe( UInterpTrack *Track, int32 KeyIndex ) const override;
+};
diff --git a/Engine/Source/Editor/Matinee/Classes/MatineeTrackVectorPropHelper.h b/Engine/Source/Editor/Matinee/Classes/MatineeTrackVectorPropHelper.h
--- a/Engine/Source/Editor/Matinee/Classes/MatineeTrackVectorPropHelper.h
+++ b/Engine/Source/Editor/Matinee/Classes/MatineeTrackVectorPropHelper.h
@@ -1,32 +1,29 @@
 // Copyright 1998-2014 Epic Games, Inc. All Rights Reserved.
 
-/**
- *
- */
-
 #pragma once
+
 #include "MatineeTrackVectorPropHelper.generated.h"
 
+
 UCLASS()
 class UMatineeTrackVectorPropHelper : public UInterpTrackHelper
 {
 	GENERATED_UCLASS_BODY()
 
-
 	/**
 	 * Pops up a dialog letting user choose between a set of properties, then checks to see if that property has been bound to yet.
 	 *
-	 * @param PropNames		Possible property names to select from.
-	 *
+	 * @param PropNames Possible property names to select from.
 	 * @return true if the property selected was acceptable, false otherwise.
 	 */
 	virtual bool ChooseProperty(TArray<FName> &PropNames) const;
 
-	// Begin UInterpTrackHelper Interface
-	virtual	bool PreCreateTrack( UInterpGroup* Group, const UInterpTrack *TrackDef, bool bDuplicatingTrack, bool bAllowPrompts ) const OVERRIDE;
-	virtual void  PostCreateTrack( UInterpTrack *Track, bool bDuplicatingTrack, int32 TrackIndex ) const OVERRIDE;
-	// End UInterpTrackHelper Interface
-
 	void OnCreateTrackTextEntry(const FString& ChosenText, TWeakPtr<SWindow> Window, FString* OutputString);
-};
 
+public:
+
+	// UInterpTrackHelper interface
+
+	virtual	bool PreCreateTrack( UInterpGroup* Group, const UInterpTrack *TrackDef, bool bDuplicatingTrack, bool bAllowPrompts ) const override;
+	virtual void  PostCreateTrack( UInterpTrack *Track, bool bDuplicatingTrack, int32 TrackIndex ) const override;
+};
diff --git a/Engine/Source/Editor/Matinee/Classes/MatineeTrackVisibilityHelper.h b/Engine/Source/Editor/Matinee/Classes/MatineeTrackVisibilityHelper.h
--- a/Engine/Source/Editor/Matinee/Classes/MatineeTrackVisibilityHelper.h
+++ b/Engine/Source/Editor/Matinee/Classes/MatineeTrackVisibilityHelper.h
@@ -9,12 +9,13 @@ class UMatineeTrackVisibilityHelper : public UInterpTrackHelper
 {
 	GENERATED_UCLASS_BODY()
 
+	void OnAddKeyTextEntry(const FString& ChosenText, IMatineeBase* Matinee, UInterpTrack* Track);
 
-	// Begin UInterpTrackHelper Interface
-	virtual	bool PreCreateKeyframe( UInterpTrack *Track, float KeyTime ) const OVERRIDE;
-	virtual void  PostCreateKeyframe( UInterpTrack *Track, int32 KeyIndex ) const OVERRIDE;
-	// End UInterpTrackHelper Interface
+public:
 
-	void OnAddKeyTextEntry(const FString& ChosenText, IMatineeBase* Matinee, UInterpTrack* Track);
+	// UInterpTrackHelper Interface
+
+	virtual	bool PreCreateKeyframe( UInterpTrack *Track, float KeyTime ) const override;
+	virtual void  PostCreateKeyframe( UInterpTrack *Track, int32 KeyIndex ) const override;
 };
 
diff --git a/Engine/Source/Editor/Matinee/Classes/MatineeTransBuffer.h b/Engine/Source/Editor/Matinee/Classes/MatineeTransBuffer.h
--- a/Engine/Source/Editor/Matinee/Classes/MatineeTransBuffer.h
+++ b/Engine/Source/Editor/Matinee/Classes/MatineeTransBuffer.h
@@ -4,32 +4,19 @@
 
 #include "MatineeTransBuffer.generated.h"
 
+
 /**
- * Special transaction buffer for Matinee undo/redo
+ * Special transaction buffer for Matinee undo/redo.
  * Will be capped at InMaxMemory.
- * 
  */
 UCLASS(transient)
 class UMatineeTransBuffer : public UTransBuffer
 {
     GENERATED_UCLASS_BODY()
+
 	UMatineeTransBuffer(const class FPostConstructInitializeProperties& PCIP, SIZE_T InMaxMemory)
 		:	UTransBuffer( PCIP, InMaxMemory )
-	{}
-
-	// Begin UTransactor Interface
-	virtual int32 Begin( const TCHAR* SessionContext, const FText& Description ) OVERRIDE
-	{
-		return 0;
-	}
-
-	virtual int32 End() OVERRIDE
-	{
-		return 1;
-	}
-
-	virtual void Cancel(int32 StartIndex = 0) OVERRIDE {}
-	// End UTransactor Interface
+	{ }
 
 	/**  
 	 * Begin a Matinee specific transaction
@@ -39,4 +26,19 @@ class UMatineeTransBuffer : public UTransBuffer
 
 	/** End a Matinee specific transaction */
 	virtual void EndSpecial();
+
+public:
+
+	// UTransactor Interface
+	virtual int32 Begin( const TCHAR* SessionContext, const FText& Description ) override
+	{
+		return 0;
+	}
+
+	virtual int32 End() override
+	{
+		return 1;
+	}
+
+	virtual void Cancel(int32 StartIndex = 0) override {}
 };
diff --git a/Engine/Source/Editor/Matinee/Private/Matinee.cpp b/Engine/Source/Editor/Matinee/Private/Matinee.cpp
--- a/Engine/Source/Editor/Matinee/Private/Matinee.cpp
+++ b/Engine/Source/Editor/Matinee/Private/Matinee.cpp
@@ -1,9 +1,5 @@
 // Copyright 1998-2014 Epic Games, Inc. All Rights Reserved.
 
-/*=============================================================================
-	Matinee.cpp: Interpolation/Cinema editing
-=============================================================================*/
-
 #include "MatineeModule.h"
 
 #include "Matinee/InterpFilter.h"
@@ -40,6 +36,7 @@
 #include "MessageLog.h"
 #include "Editor/PropertyEditor/Public/IDetailsView.h"
 
+
 DEFINE_LOG_CATEGORY(LogSlateMatinee);
 
 #define LOCTEXT_NAMESPACE "Matinee"
diff --git a/Engine/Source/Editor/Matinee/Private/MatineeCurveEd.cpp b/Engine/Source/Editor/Matinee/Private/MatineeCurveEd.cpp
--- a/Engine/Source/Editor/Matinee/Private/MatineeCurveEd.cpp
+++ b/Engine/Source/Editor/Matinee/Private/MatineeCurveEd.cpp
@@ -1,11 +1,5 @@
-
 // Copyright 1998-2014 Epic Games, Inc. All Rights Reserved.
 
-/*=============================================================================
-  MatineeCurveEd.cpp: Functionality pertaining to the Curve Editor in Matinee
-=============================================================================*/
-
-
 #include "MatineeModule.h"
 
 #include "Matinee/MatineeActor.h"
@@ -15,6 +9,7 @@
 
 #include "Matinee.h"
 
+
 /** Ensure the curve editor is synchronised with the track editor. */
 void FMatinee::SyncCurveEdView()
 {
diff --git a/Engine/Source/Editor/Matinee/Private/MatineeDraw.cpp b/Engine/Source/Editor/Matinee/Private/MatineeDraw.cpp
--- a/Engine/Source/Editor/Matinee/Private/MatineeDraw.cpp
+++ b/Engine/Source/Editor/Matinee/Private/MatineeDraw.cpp
@@ -1,9 +1,5 @@
 // Copyright 1998-2014 Epic Games, Inc. All Rights Reserved.
 
-/*=============================================================================
-	MatineeDraw.cpp: Functions covering drawing the Matinee window
-=============================================================================*/
-
 #include "MatineeModule.h"
 #include "Matinee.h"
 #include "MatineeViewportClient.h"
@@ -25,6 +21,7 @@
 #include "Animation/SkeletalMeshActor.h"
 #include "Particles/Emitter.h"
 
+
 static const int32 GroupHeadHeight = 24;
 static const int32 TrackHeight = 24;
 static const int32 SubTrackHeight = 19;
diff --git a/Engine/Source/Editor/Matinee/Private/MatineeKeyReduction.cpp b/Engine/Source/Editor/Matinee/Private/MatineeKeyReduction.cpp
--- a/Engine/Source/Editor/Matinee/Private/MatineeKeyReduction.cpp
+++ b/Engine/Source/Editor/Matinee/Private/MatineeKeyReduction.cpp
@@ -1,9 +1,5 @@
 // Copyright 1998-2014 Epic Games, Inc. All Rights Reserved.
 
-/*=============================================================================
-	MatineeKeyReduction.cpp: KeyReduction dialog definitions
-=============================================================================*/
-
 #include "MatineeModule.h"
 #include "Matinee/MatineeActor.h"
 #include "Matinee/InterpGroupInst.h"
@@ -11,6 +7,7 @@
 #include "Slate.h"
 #include "SNumericEntryBox.h"
 
+
 /**
  * Dialog that requests the key reduction parameters from the user.
  * These parameters are tolerance and the reduction interval.
diff --git a/Engine/Source/Editor/Matinee/Private/MatineeMenus.cpp b/Engine/Source/Editor/Matinee/Private/MatineeMenus.cpp
--- a/Engine/Source/Editor/Matinee/Private/MatineeMenus.cpp
+++ b/Engine/Source/Editor/Matinee/Private/MatineeMenus.cpp
@@ -1,9 +1,5 @@
 // Copyright 1998-2014 Epic Games, Inc. All Rights Reserved.
 
-/*=============================================================================
-	InterpEditorMenus.cpp: Interpolation editing menus
-=============================================================================*/
-
 #include "MatineeModule.h"
 
 #include "Matinee/InterpTrackInst.h"
@@ -46,6 +42,7 @@
 #include "IDocumentation.h"
 #include "MainFrame.h"
 
+
 #define LOCTEXT_NAMESPACE "MatineeMenus"
 
 namespace
diff --git a/Engine/Source/Editor/Matinee/Private/MatineeRecordMovie.cpp b/Engine/Source/Editor/Matinee/Private/MatineeRecordMovie.cpp
--- a/Engine/Source/Editor/Matinee/Private/MatineeRecordMovie.cpp
+++ b/Engine/Source/Editor/Matinee/Private/MatineeRecordMovie.cpp
@@ -1,13 +1,10 @@
 // Copyright 1998-2014 Epic Games, Inc. All Rights Reserved.
 
-/*=============================================================================
-	InterpEditorMenus.cpp: Interpolation editing menus
-=============================================================================*/
-
 #include "MatineeModule.h"
 #include "Matinee/MatineeActor.h"
 #include "Matinee.h"
 
+
 #define LOCTEXT_NAMESPACE "MatineeMovieCapture"
 
 /**
diff --git a/Engine/Source/Editor/Matinee/Private/MatineeTools.cpp b/Engine/Source/Editor/Matinee/Private/MatineeTools.cpp
--- a/Engine/Source/Editor/Matinee/Private/MatineeTools.cpp
+++ b/Engine/Source/Editor/Matinee/Private/MatineeTools.cpp
@@ -1,10 +1,5 @@
 // Copyright 1998-2014 Epic Games, Inc. All Rights Reserved.
 
-/*=============================================================================
-	MatineeTools.cpp: Interpolation editing support tools
-=============================================================================*/
-
-
 #include "MatineeModule.h"
 #include "Matinee.h"
 #include "MatineeClasses.h"
@@ -39,6 +34,7 @@
 #include "SoundDefinitions.h"
 #include "EditorSupportDelegates.h"
 
+
 ///// UTILS
 
 void FMatinee::TickInterp(float DeltaTime)
diff --git a/Engine/Source/Editor/Matinee/Private/MatineeTrackView.cpp b/Engine/Source/Editor/Matinee/Private/MatineeTrackView.cpp
--- a/Engine/Source/Editor/Matinee/Private/MatineeTrackView.cpp
+++ b/Engine/Source/Editor/Matinee/Private/MatineeTrackView.cpp
@@ -1,10 +1,5 @@
-
 // Copyright 1998-2014 Epic Games, Inc. All Rights Reserved.
 
-/*=============================================================================
-	MatineeTrackView.cpp: Matinee Track View/Window functionality
-=============================================================================*/
-
 #include "MatineeModule.h"
 #include "Matinee.h"
 
diff --git a/Engine/Source/Editor/Matinee/Private/MatineeTransaction.h b/Engine/Source/Editor/Matinee/Private/MatineeTransaction.h
--- a/Engine/Source/Editor/Matinee/Private/MatineeTransaction.h
+++ b/Engine/Source/Editor/Matinee/Private/MatineeTransaction.h
@@ -8,13 +8,16 @@
 class FMatineeTransaction : public FTransaction
 {
 public:
+
 	FMatineeTransaction( const FText& InTitle=FText(), bool InFlip=0 )
 		:	FTransaction(TEXT("Matinee"), InTitle, InFlip)
-	{}
+	{ }
+
+public:
+
+	// FTransaction interface
 
-	// Begin FTransaction Interface
-	virtual void SaveObject( UObject* Object ) OVERRIDE;
-	virtual void SaveArray( UObject* Object, FScriptArray* Array, int32 Index, int32 Count, int32 Oper, int32 ElementSize, STRUCT_DC DefaultConstructor, STRUCT_AR Serializer, STRUCT_DTOR Destructor ) OVERRIDE;
-	// End FTransaction Interface
+	virtual void SaveObject( UObject* Object ) override;
+	virtual void SaveArray( UObject* Object, FScriptArray* Array, int32 Index, int32 Count, int32 Oper, int32 ElementSize, STRUCT_DC DefaultConstructor, STRUCT_AR Serializer, STRUCT_DTOR Destructor ) override;
 };
 
diff --git a/Engine/Source/Editor/MeshPaint/Private/MeshPaintEdMode.cpp b/Engine/Source/Editor/MeshPaint/Private/MeshPaintEdMode.cpp
--- a/Engine/Source/Editor/MeshPaint/Private/MeshPaintEdMode.cpp
+++ b/Engine/Source/Editor/MeshPaint/Private/MeshPaintEdMode.cpp
@@ -1,9 +1,5 @@
 // Copyright 1998-2014 Epic Games, Inc. All Rights Reserved.
 
-/*================================================================================
-	MeshPaintEdMode.cpp: Mesh paint tool
-================================================================================*/
-
 #include "MeshPaintPrivatePCH.h"
 
 #include "Materials/MaterialExpressionTextureBase.h"
@@ -32,6 +28,7 @@
 #include "Runtime/Engine/Classes/PhysicsEngine/BodySetup.h"
 #include "SMeshPaint.h"
 
+
 #define LOCTEXT_NAMESPACE "MeshPaint_Mode"
 
 DEFINE_LOG_CATEGORY_STATIC(LogMeshPaintEdMode, Log, All);
diff --git a/Engine/Source/Editor/MeshPaint/Public/MeshPaintEdMode.h b/Engine/Source/Editor/MeshPaint/Public/MeshPaintEdMode.h
--- a/Engine/Source/Editor/MeshPaint/Public/MeshPaintEdMode.h
+++ b/Engine/Source/Editor/MeshPaint/Public/MeshPaintEdMode.h
@@ -1,13 +1,5 @@
 // Copyright 1998-2014 Epic Games, Inc. All Rights Reserved.
 
-/*================================================================================
-	MeshPaintEdMode.h: Mesh paint tool
-================================================================================*/
-
-
-#ifndef __MeshPaintEdMode_h__
-#define __MeshPaintEdMode_h__
-
 #pragma once
 
 #include "GenericOctree.h"
@@ -27,7 +19,6 @@ namespace EMeshPaintResource
 }
 
 
-
 /** Mesh paint mode */
 namespace EMeshPaintMode
 {
@@ -42,7 +33,6 @@ namespace EMeshPaintMode
 }
 
 
-
 /** Vertex paint target */
 namespace EMeshVertexPaintTarget
 {
@@ -57,8 +47,6 @@ namespace EMeshVertexPaintTarget
 }
 
 
-
-
 /** Mesh paint color view modes (somewhat maps to EVertexColorViewMode engine enum.) */
 namespace EMeshPaintColorViewMode
 {
@@ -90,7 +78,6 @@ namespace EMeshPaintColorViewMode
  */
 class FMeshPaintSettings
 {
-
 public:
 
 	/** Static: Returns global mesh paint settings */
@@ -99,14 +86,11 @@ public:
 		return StaticMeshPaintSettings;
 	}
 
-
 protected:
 
 	/** Static: Global mesh paint settings */
 	static FMeshPaintSettings StaticMeshPaintSettings;
 
-
-
 public:
 
 	/** Constructor */
@@ -131,9 +115,7 @@ public:
 		  UVChannel( 0 ),
 		  ColorViewMode( EMeshPaintColorViewMode::Normal ),
 		  bEnableSeamPainting( true )
-	{
-	}
-
+	{ }
 
 public:
 
@@ -155,12 +137,10 @@ public:
 	/** Resource type we're editing */
 	EMeshPaintResource::Type ResourceType;
 
-
 	/** Mode we're painting in.  If this is set to PaintColors we're painting colors directly; if it's set
 	    to PaintWeights we're painting texture blend weights using a single channel. */
 	EMeshPaintMode::Type PaintMode;
 
-
 	/** Colors Mode: Paint color */
 	FLinearColor PaintColor;
 
@@ -179,7 +159,6 @@ public:
 	/** Colors Mode: True if alpha colors values should be written */
 	bool bWriteAlpha;
 
-
 	/** Weights Mode: Total weight count */
 	int32 TotalWeightCount;
 
@@ -189,7 +168,6 @@ public:
 	/** Weights Mode: Weight index that we're currently using to erase with */
 	int32 EraseWeightIndex;
 
-
 	/**
 	 * Vertex paint settings
 	 */
@@ -197,7 +175,6 @@ public:
 	/** Vertex paint target */
 	EMeshVertexPaintTarget::Type VertexPaintTarget;
 
-
 	/**
 	 * Texture paint settings
 	 */
@@ -205,7 +182,6 @@ public:
 	/** UV channel to paint textures using */
 	int32 UVChannel;
 
-
 	/**
 	 * View settings
 	 */
@@ -218,7 +194,6 @@ public:
 };
 
 
-
 /** Mesh painting action (paint, erase) */
 namespace EMeshPaintAction
 {
@@ -239,7 +214,6 @@ namespace EMeshPaintAction
 }
 
 
-
 namespace MeshPaintDefs
 {
 	// Design constraints
@@ -267,6 +241,7 @@ struct FTextureTargetListInfo
 	{}
 };
 
+
 /** 
  *  Wrapper to store which of a meshes materials is selected as well as the total number of materials. 
  */
@@ -281,6 +256,7 @@ struct FMeshSelectedMaterialInfo
 	{}
 };
 
+
 /**
  * Mesh Paint editor mode
  */
@@ -747,6 +723,7 @@ private:
 
 	/** Caches the currently selected actors info into CurrentlySelectedActorsMaterialInfo */
 	void CacheActorInfo();
+
 private:
 
 	/** Whether we're currently painting */
@@ -832,43 +809,53 @@ private:
 
 	/** The currently selected actor, used to refer into the Map of Selected actor info */
 	TWeakObjectPtr<AActor> ActorBeingEdited;
-
 };
 
 
-
 /**
 * Imports vertex colors from tga file onto selected meshes
 */
 class FImportVertexTextureHelper  
 {
 public:
-	enum ChannelsMask{
+
+	/**
+	 * Enumerates channel masks.
+	 */
+	enum ChannelsMask
+	{
+		/** Red channel. */
 		ERed		= 0x1,
+
+		/** Green channel. */
 		EGreen		= 0x2,
+
+		/** Blue channel. */
 		EBlue		= 0x4,
+
+		/** Alpha channel. */
 		EAlpha		= 0x8,
 	};
 
 	/**
-	* PickVertexColorFromTex() - Color picker function. Retrieves pixel color from coordinates and mask.
-	* @param NewVertexColor - returned color
-	* @param MipData - Highest mip-map with pixels data
-	* @param UV - texture coordinate to read
-	* @param Tex - texture info
-	* @param ColorMask - mask for filtering which colors to use
-	*/
+	 * Color picker function. Retrieves pixel color from coordinates and mask.
+	 *
+	 * @param NewVertexColor Returned color.
+	 * @param MipData Highest mip-map with pixels data.
+	 * @param UV Texture coordinate to read.
+	 * @param Tex Texture info.
+	 * @param ColorMask Mask for filtering which colors to use.
+	 */
 	void PickVertexColorFromTex(FColor & NewVertexColor, uint8* MipData, FVector2D & UV, UTexture2D* Tex, uint8 & ColorMask);
 	
 	/**
-	* Imports Vertex Color data from texture scanning thought uv vertex coordinates for selected actors.  
-	* @param Filename - path for loading TGA file
-	* @param UVIndex - Coordinate index
-	* @param ImportLOD - LOD level to work with
-	* @param Tex - texture info
-	* @param ColorMask - mask for filtering which colors to use
-	*/
+	 * Imports Vertex Color data from texture scanning thought uv vertex coordinates for selected actors.  
+	 *
+	 * @param Filename path for loading TGA file.
+	 * @param UVIndex Coordinate index.
+	 * @param ImportLOD LOD level to work with.
+	 * @param Tex Texture info.
+	 * @param ColorMask Mask for filtering which colors to use.
+	 */
 	void ImportVertexColors(const FString & Filename, int32 UVIndex, int32 ImportLOD, uint8 ColorMask);
-
 };
-#endif	// __MeshPaintEdMode_h__
diff --git a/Engine/Source/Editor/NewsFeed/Private/Implementation/CdnNewsFeedTitleFile.h b/Engine/Source/Editor/NewsFeed/Private/Implementation/CdnNewsFeedTitleFile.h
--- a/Engine/Source/Editor/NewsFeed/Private/Implementation/CdnNewsFeedTitleFile.h
+++ b/Engine/Source/Editor/NewsFeed/Private/Implementation/CdnNewsFeedTitleFile.h
@@ -40,23 +40,15 @@ public:
 
 public:
 
-	// Begin IOnlineTitleFile interface
-
-	virtual bool GetFileContents(const FString& DLName, TArray<uint8>& FileContents) OVERRIDE;
-
-	virtual bool ClearFiles() OVERRIDE;
-
-	virtual bool ClearFile(const FString& DLName) OVERRIDE;
-
-	virtual bool EnumerateFiles() OVERRIDE;
-
-	virtual bool EnumerateFiles(int32 Start, int32 Count) OVERRIDE;
-
-	virtual void GetFileList(TArray<FCloudFileHeader>& InFileHeaders) OVERRIDE;
-
-	virtual bool ReadFile(const FString& DLName) OVERRIDE;
-
-	// Begin IOnlineTitleFile interface
+	// IOnlineTitleFile interface
+
+	virtual bool GetFileContents(const FString& DLName, TArray<uint8>& FileContents) override;
+	virtual bool ClearFiles() override;
+	virtual bool ClearFile(const FString& DLName) override;
+	virtual bool EnumerateFiles() override;
+	virtual bool EnumerateFiles(int32 Start, int32 Count) override;
+	virtual void GetFileList(TArray<FCloudFileHeader>& InFileHeaders) override;
+	virtual bool ReadFile(const FString& DLName) override;
 
 protected:
 	
diff --git a/Engine/Source/Editor/NewsFeed/Private/Implementation/LocalNewsFeedTitleFile.h b/Engine/Source/Editor/NewsFeed/Private/Implementation/LocalNewsFeedTitleFile.h
--- a/Engine/Source/Editor/NewsFeed/Private/Implementation/LocalNewsFeedTitleFile.h
+++ b/Engine/Source/Editor/NewsFeed/Private/Implementation/LocalNewsFeedTitleFile.h
@@ -15,23 +15,15 @@ public:
 
 public:
 
-	// Begin IOnlineTitleFile interface
-
-	virtual bool GetFileContents(const FString& DLName, TArray<uint8>& FileContents) OVERRIDE;
-
-	virtual bool ClearFiles() OVERRIDE;
-
-	virtual bool ClearFile(const FString& DLName) OVERRIDE;
-
-	virtual bool EnumerateFiles() OVERRIDE;
-
-	virtual bool EnumerateFiles(int32 Start, int32 Count) OVERRIDE;
-
-	virtual void GetFileList(TArray<FCloudFileHeader>& InFileHeaders) OVERRIDE;
-
-	virtual bool ReadFile(const FString& DLName) OVERRIDE;
-
-	// End IOnlineTitleFile interface
+	// IOnlineTitleFile interface
+
+	virtual bool GetFileContents(const FString& DLName, TArray<uint8>& FileContents) override;
+	virtual bool ClearFiles() override;
+	virtual bool ClearFile(const FString& DLName) override;
+	virtual bool EnumerateFiles() override;
+	virtual bool EnumerateFiles(int32 Start, int32 Count) override;
+	virtual void GetFileList(TArray<FCloudFileHeader>& InFileHeaders) override;
+	virtual bool ReadFile(const FString& DLName) override;
 
 protected:
 
diff --git a/Engine/Source/Editor/NewsFeed/Private/Implementation/NewsFeedSettings.h b/Engine/Source/Editor/NewsFeed/Private/Implementation/NewsFeedSettings.h
--- a/Engine/Source/Editor/NewsFeed/Private/Implementation/NewsFeedSettings.h
+++ b/Engine/Source/Editor/NewsFeed/Private/Implementation/NewsFeedSettings.h
@@ -1,9 +1,5 @@
 // Copyright 1998-2014 Epic Games, Inc. All Rights Reserved.
 
-/*=============================================================================
-	NewsFeedSettings.h: Declares the UNewsFeedSettings class.
-=============================================================================*/
-
 #pragma once
 
 #include "NewsFeedSettings.generated.h"
@@ -37,15 +33,11 @@ class UNewsFeedSettings
 
 public:
 
-	/**
-	 * The URL at which the news feed data files are located when using the Source=Cdn.
-	 */
+	/** The URL at which the news feed data files are located when using the Source=Cdn. */
 	UPROPERTY(config)
 	FString CdnSourceUrl;
 
-	/**
-	 * The path to the local data files when using Source=Local.
-	 */
+	/** The path to the local data files when using Source=Local. */
 	UPROPERTY(config)
 	FString LocalSourcePath;
 
@@ -59,21 +51,15 @@ public:
 
 public:
 
-	/**
-	 * The maximum number of news items to show.
-	 */
+	/** The maximum number of news items to show. */
 	UPROPERTY(config, EditAnywhere, Category=Display, meta=(DisplayName="Most Recent Items To Show"))
 	int32 MaxItemsToShow;
 
-	/**
-	 * List of news feed items that have been marked as read.
-	 */
+	/** List of news feed items that have been marked as read. */
 	UPROPERTY(config)
 	TArray<FGuid> ReadItems;
 
-	/**
-	 * Whether to show only unread news feed items.
-	 */
+	/** Whether to show only unread news feed items. */
 	UPROPERTY(config, EditAnywhere, Category=Display)
 	bool ShowOnlyUnreadItems;
 
@@ -85,15 +71,16 @@ public:
 	 * @return The delegate.
 	 */
 	DECLARE_EVENT_OneParam(UNewsFeedSettings, FSettingChangedEvent, FName /*PropertyName*/);
-	FSettingChangedEvent& OnSettingChanged( ) { return SettingChangedEvent; }
+	FSettingChangedEvent& OnSettingChanged( )
+	{
+		return SettingChangedEvent;
+	}
 
 protected:
 
-	// Begin UObject overrides
-
-	virtual void PostEditChangeProperty( struct FPropertyChangedEvent& PropertyChangedEvent ) OVERRIDE;
+	// UObject overrides
 
-	// End UObject overrides
+	virtual void PostEditChangeProperty( struct FPropertyChangedEvent& PropertyChangedEvent ) override;
 
 private:
 
diff --git a/Engine/Source/Editor/NewsFeed/Private/Models/NewsFeedItem.h b/Engine/Source/Editor/NewsFeed/Private/Models/NewsFeedItem.h
--- a/Engine/Source/Editor/NewsFeed/Private/Models/NewsFeedItem.h
+++ b/Engine/Source/Editor/NewsFeed/Private/Models/NewsFeedItem.h
@@ -1,9 +1,5 @@
 // Copyright 1998-2014 Epic Games, Inc. All Rights Reserved.
 
-/*=============================================================================
-	NewsFeedItem.h: Declares the FNewsFeedItem class.
-=============================================================================*/
-
 #pragma once
 
 
@@ -48,7 +44,7 @@ public:
 	/**
 	 * Creates and initializes a new item with the specified identifier.
 	 *
-	 * @param InItemId - The item identifier.
+	 * @param InItemId The item identifier.
 	 */
 	FNewsFeedItem( const FGuid& InItemId )
 		: ItemId(InItemId)
diff --git a/Engine/Source/Editor/NewsFeed/Private/NewsFeedClasses.cpp b/Engine/Source/Editor/NewsFeed/Private/NewsFeedClasses.cpp
--- a/Engine/Source/Editor/NewsFeed/Private/NewsFeedClasses.cpp
+++ b/Engine/Source/Editor/NewsFeed/Private/NewsFeedClasses.cpp
@@ -1,12 +1,9 @@
 // Copyright 1998-2014 Epic Games, Inc. All Rights Reserved.
 
-/*=============================================================================
-	NewsFeedClasses.cpp: Implements the module's script classes.
-=============================================================================*/
-
 #include "NewsFeedPrivatePCH.h"
 
-/* UNewsFeedSettings interface
+
+/* UNewsFeedSettings structors
  *****************************************************************************/
 
 UNewsFeedSettings::UNewsFeedSettings( const class FPostConstructInitializeProperties& PCIP )
@@ -16,6 +13,9 @@ UNewsFeedSettings::UNewsFeedSettings( const class FPostConstructInitializeProper
 { }
 
 
+/* UObject overrides
+ *****************************************************************************/
+
 void UNewsFeedSettings::PostEditChangeProperty( struct FPropertyChangedEvent& PropertyChangedEvent )
 {
 	Super::PostEditChangeProperty(PropertyChangedEvent);
diff --git a/Engine/Source/Editor/NewsFeed/Private/NewsFeedModule.cpp b/Engine/Source/Editor/NewsFeed/Private/NewsFeedModule.cpp
--- a/Engine/Source/Editor/NewsFeed/Private/NewsFeedModule.cpp
+++ b/Engine/Source/Editor/NewsFeed/Private/NewsFeedModule.cpp
@@ -1,9 +1,5 @@
 // Copyright 1998-2014 Epic Games, Inc. All Rights Reserved.
 
-/*=============================================================================
-	NewsFeedModule.cpp: Implements the FNewsFeedModule class.
-=============================================================================*/
-
 #include "NewsFeedPrivatePCH.h"
 
 
@@ -18,9 +14,9 @@ class FNewsFeedModule
 {
 public:
 
-	// Begin IModuleInterface interface
+	// IModuleInterface interface
 
-	virtual void StartupModule( ) OVERRIDE
+	virtual void StartupModule( ) override
 	{
 		// register settings
 		ISettingsModule* SettingsModule = ISettingsModule::Get();
@@ -35,7 +31,7 @@ public:
 		}
 	}
 
-	virtual void ShutdownModule( ) OVERRIDE
+	virtual void ShutdownModule( ) override
 	{
 		// unregister settings
 		ISettingsModule* SettingsModule = ISettingsModule::Get();
@@ -46,13 +42,11 @@ public:
 		}
 	}
 
-	// End IModuleInterface interface
-
 public:
 
-	// Begin INewsFeedModule interface
+	// INewsFeedModule interface
 
-	virtual TSharedRef<class SWidget> CreateNewsFeedButton( ) OVERRIDE
+	virtual TSharedRef<class SWidget> CreateNewsFeedButton( ) override
 	{
 		FNewsFeedCacheRef NewsFeedCache = MakeShareable(new FNewsFeedCache());
 
@@ -106,8 +100,6 @@ public:
 			];
 	}
 
-	// End INewsFeedModule interface
-
 private:
 
 	// Callback for getting the background color the unread item count.
diff --git a/Engine/Source/Editor/NewsFeed/Private/NewsFeedPrivatePCH.h b/Engine/Source/Editor/NewsFeed/Private/NewsFeedPrivatePCH.h
--- a/Engine/Source/Editor/NewsFeed/Private/NewsFeedPrivatePCH.h
+++ b/Engine/Source/Editor/NewsFeed/Private/NewsFeedPrivatePCH.h
@@ -1,15 +1,11 @@
 // Copyright 1998-2014 Epic Games, Inc. All Rights Reserved.
 
-/*=============================================================================
-	NewsFeedPrivatePCH.h: Pre-compiled header file for the NewsFeed module.
-=============================================================================*/
-
 #pragma once
 
-#include "../Public/NewsFeed.h"
+#include "NewsFeed.h"
 
 
-/* Dependencies
+/* Private dependencies
  *****************************************************************************/
 
 #include "AnalyticsEventAttribute.h"
diff --git a/Engine/Source/Editor/NewsFeed/Private/Widgets/SNewsFeed.cpp b/Engine/Source/Editor/NewsFeed/Private/Widgets/SNewsFeed.cpp
--- a/Engine/Source/Editor/NewsFeed/Private/Widgets/SNewsFeed.cpp
+++ b/Engine/Source/Editor/NewsFeed/Private/Widgets/SNewsFeed.cpp
@@ -1,9 +1,5 @@
 // Copyright 1998-2014 Epic Games, Inc. All Rights Reserved.
 
-/*=============================================================================
-	SNewsFeed.cpp: Implements the SNewsFeed class.
-=============================================================================*/
-
 #include "NewsFeedPrivatePCH.h"
 
 
diff --git a/Engine/Source/Editor/NewsFeed/Private/Widgets/SNewsFeed.h b/Engine/Source/Editor/NewsFeed/Private/Widgets/SNewsFeed.h
--- a/Engine/Source/Editor/NewsFeed/Private/Widgets/SNewsFeed.h
+++ b/Engine/Source/Editor/NewsFeed/Private/Widgets/SNewsFeed.h
@@ -1,9 +1,5 @@
 // Copyright 1998-2014 Epic Games, Inc. All Rights Reserved.
 
-/*=============================================================================
-	SNewsFeed.h: Declares the SNewsFeed class.
-=============================================================================*/
-
 #pragma once
 
 
diff --git a/Engine/Source/Editor/NewsFeed/Private/Widgets/SNewsFeedListRow.h b/Engine/Source/Editor/NewsFeed/Private/Widgets/SNewsFeedListRow.h
--- a/Engine/Source/Editor/NewsFeed/Private/Widgets/SNewsFeedListRow.h
+++ b/Engine/Source/Editor/NewsFeed/Private/Widgets/SNewsFeedListRow.h
@@ -1,9 +1,5 @@
 // Copyright 1998-2014 Epic Games, Inc. All Rights Reserved.
 
-/*=============================================================================
-	SNewsFeedListRow.h: Declares the SNewsFeedListRow class.
-=============================================================================*/
-
 #pragma once
 
 
diff --git a/Engine/Source/Editor/NewsFeed/Public/Interfaces/INewsFeedModule.h b/Engine/Source/Editor/NewsFeed/Public/Interfaces/INewsFeedModule.h
--- a/Engine/Source/Editor/NewsFeed/Public/Interfaces/INewsFeedModule.h
+++ b/Engine/Source/Editor/NewsFeed/Public/Interfaces/INewsFeedModule.h
@@ -1,9 +1,5 @@
 // Copyright 1998-2014 Epic Games, Inc. All Rights Reserved.
 
-/*=============================================================================
-	INewsFeedModule.h: Declares the INewsFeedModule interface.
-=============================================================================*/
-
 #pragma once
 
 
diff --git a/Engine/Source/Editor/NewsFeed/Public/NewsFeed.h b/Engine/Source/Editor/NewsFeed/Public/NewsFeed.h
--- a/Engine/Source/Editor/NewsFeed/Public/NewsFeed.h
+++ b/Engine/Source/Editor/NewsFeed/Public/NewsFeed.h
@@ -1,9 +1,5 @@
 // Copyright 1998-2014 Epic Games, Inc. All Rights Reserved.
 
-/*=============================================================================
-	NewsFeed.h: NewsFeed module public header file.
-=============================================================================*/
-
 #pragma once
 
 
diff --git a/Engine/Source/Editor/PListEditor/Public/SPlistEditor.h b/Engine/Source/Editor/PListEditor/Public/SPlistEditor.h
--- a/Engine/Source/Editor/PListEditor/Public/SPlistEditor.h
+++ b/Engine/Source/Editor/PListEditor/Public/SPlistEditor.h
@@ -1,14 +1,7 @@
 // Copyright 1998-2014 Epic Games, Inc. All Rights Reserved.
 
-/*=============================================================================
-	      SPlistEditor.h: Definition for PList editor app/menu
-=============================================================================*/
-
 #pragma once
 
-#ifndef __SPListEditor_h__
-#define __SPListEditor_h__
-
 #include "Slate.h"
 #include "XmlParser.h"
 #include "../Private/PListNode.h"
@@ -18,6 +11,7 @@
 #include "../Private/PListNodeFile.h"
 #include "../Private/PListNodeString.h"
 
+
 // Type of notification to spawn
 enum ENTF_Types
 {
@@ -265,5 +259,3 @@ private:
 	/** Helper function to search through nodes to find a specific node's parent */
 	bool FindParent(const TSharedPtr<IPListNode>& InChildNode, TSharedPtr<IPListNode>& OutFoundNode) const;
 };
-
-#endif
diff --git a/Engine/Source/Editor/ProjectSettingsViewer/Private/ProjectSettingsViewerModule.cpp b/Engine/Source/Editor/ProjectSettingsViewer/Private/ProjectSettingsViewerModule.cpp
--- a/Engine/Source/Editor/ProjectSettingsViewer/Private/ProjectSettingsViewerModule.cpp
+++ b/Engine/Source/Editor/ProjectSettingsViewer/Private/ProjectSettingsViewerModule.cpp
@@ -1,13 +1,10 @@
 // Copyright 1998-2014 Epic Games, Inc. All Rights Reserved.
 
-/*=============================================================================
-	ProjectSettingsViewerModule.cpp: Implements the FProjectSettingsViewerModule class.
-=============================================================================*/
-
 #include "ProjectSettingsViewerPrivatePCH.h"
 #include "Engine/Console.h"
 #include "ProjectTargetPlatformEditor.h"
 
+
 #define LOCTEXT_NAMESPACE "FProjectSettingsViewerModule"
 
 static const FName ProjectSettingsTabName("ProjectSettings");
@@ -22,9 +19,9 @@ class FProjectSettingsViewerModule
 {
 public:
 
-	// Begin ISettingsViewer interface
+	// ISettingsViewer interface
 
-	virtual void ShowSettings( const FName& CategoryName, const FName& SectionName ) OVERRIDE
+	virtual void ShowSettings( const FName& CategoryName, const FName& SectionName ) override
 	{
 		FGlobalTabmanager::Get()->InvokeTab(ProjectSettingsTabName);
 		ISettingsEditorModelPtr SettingsEditorModel = SettingsEditorModelPtr.Pin();
@@ -40,13 +37,11 @@ public:
 		}
 	}
 
-	// End ISettingsViewer interface
-
 public:
 
-	// Begin IModuleInterface interface
+	// IModuleInterface interface
 
-	virtual void StartupModule( ) OVERRIDE
+	virtual void StartupModule( ) override
 	{
 		ISettingsModule* SettingsModule = ISettingsModule::Get();
 
@@ -63,19 +58,17 @@ public:
 			.SetMenuType(ETabSpawnerMenuType::Hide);
 	}
 
-	virtual void ShutdownModule( ) OVERRIDE
+	virtual void ShutdownModule( ) override
 	{
 		FGlobalTabmanager::Get()->UnregisterNomadTabSpawner(ProjectSettingsTabName);
 		UnregisterSettings();
 	}
 
-	virtual bool SupportsDynamicReloading( ) OVERRIDE
+	virtual bool SupportsDynamicReloading( ) override
 	{
 		return true;
 	}
 
-	// End IModuleInterface interface
-
 protected:
 
 	/**
diff --git a/Engine/Source/Editor/ProjectSettingsViewer/Private/ProjectSettingsViewerPrivatePCH.h b/Engine/Source/Editor/ProjectSettingsViewer/Private/ProjectSettingsViewerPrivatePCH.h
--- a/Engine/Source/Editor/ProjectSettingsViewer/Private/ProjectSettingsViewerPrivatePCH.h
+++ b/Engine/Source/Editor/ProjectSettingsViewer/Private/ProjectSettingsViewerPrivatePCH.h
@@ -1,9 +1,5 @@
 // Copyright 1998-2014 Epic Games, Inc. All Rights Reserved.
 
-/*=============================================================================
-	ProjectSettingsViewerPrivatePCH.h: Pre-compiled header file for the ProjectSettingsViewer module.
-=============================================================================*/
-
 #pragma  once
 
 
@@ -16,7 +12,3 @@
 #include "SettingsEditor.h"
 #include "UnrealEd.h"
 #include "MoviePlayerSettings.h"
-
-
-/* Private includes
- *****************************************************************************/
diff --git a/Engine/Source/Editor/PropertyEditor/Private/Presentation/PropertyTable/PropertyTablePropertyNameCell.cpp b/Engine/Source/Editor/PropertyEditor/Private/Presentation/PropertyTable/PropertyTablePropertyNameCell.cpp
--- a/Engine/Source/Editor/PropertyEditor/Private/Presentation/PropertyTable/PropertyTablePropertyNameCell.cpp
+++ b/Engine/Source/Editor/PropertyEditor/Private/Presentation/PropertyTable/PropertyTablePropertyNameCell.cpp
@@ -1,9 +1,5 @@
 // Copyright 1998-2014 Epic Games, Inc. All Rights Reserved.
 
-/*=============================================================================
- PropertyTablePropertyNameCell.cpp: Implements the FPropertyTablePropertyNameCell class.
-=============================================================================*/
-
 #include "PropertyEditorPrivatePCH.h"
 #include "PropertyTablePropertyNameCell.h"
 #include "PropertyTablePropertyNameColumn.h"
@@ -16,6 +12,7 @@
 #include "ItemPropertyNode.h"
 #include "PropertyHandle.h"
 
+
 FPropertyTablePropertyNameCell::FPropertyTablePropertyNameCell( const TSharedRef< class FPropertyTablePropertyNameColumn >& InColumn, const TSharedRef< class IPropertyTableRow >& InRow )
 	: EnteredEditModeEvent()
 	, ExitedEditModeEvent()
diff --git a/Engine/Source/Editor/PropertyEditor/Private/Presentation/PropertyTable/PropertyTablePropertyNameColumn.cpp b/Engine/Source/Editor/PropertyEditor/Private/Presentation/PropertyTable/PropertyTablePropertyNameColumn.cpp
--- a/Engine/Source/Editor/PropertyEditor/Private/Presentation/PropertyTable/PropertyTablePropertyNameColumn.cpp
+++ b/Engine/Source/Editor/PropertyEditor/Private/Presentation/PropertyTable/PropertyTablePropertyNameColumn.cpp
@@ -1,15 +1,13 @@
 // Copyright 1998-2014 Epic Games, Inc. All Rights Reserved.
 
-/*=============================================================================
- PropertyTablePropertyNameColumn.cpp: Implements the FPropertyTablePropertyNameColumn class.
-=============================================================================*/
-
 #include "PropertyEditorPrivatePCH.h"
 #include "PropertyTableCell.h"
 #include "PropertyTablePropertyNameColumn.h"
 
+
 #define LOCTEXT_NAMESPACE "PropertyNameColumnHeader"
 
+
 FPropertyTablePropertyNameColumn::FPropertyTablePropertyNameColumn( const TSharedRef< IPropertyTable >& InTable )
 	: Table( InTable )
 	, Cells()
diff --git a/Engine/Source/Editor/SceneOutliner/Private/SceneOutlinerModule.cpp b/Engine/Source/Editor/SceneOutliner/Private/SceneOutlinerModule.cpp
--- a/Engine/Source/Editor/SceneOutliner/Private/SceneOutlinerModule.cpp
+++ b/Engine/Source/Editor/SceneOutliner/Private/SceneOutlinerModule.cpp
@@ -1,9 +1,5 @@
 // Copyright 1998-2014 Epic Games, Inc. All Rights Reserved.
 
-/*=============================================================================
-	SceneOutlinerModule.cpp: Implements the FSceneOutlinerModule class.
-=============================================================================*/
-
 #include "SceneOutlinerPrivatePCH.h"
 #include "SceneOutlinerTreeItems.h"
 #include "ModuleManager.h"
@@ -34,6 +30,7 @@ TSharedRef< ISceneOutliner > FSceneOutlinerModule::CreateSceneOutliner( const FS
 	return CreateSceneOutliner(InitOptions, OnItemPicked);
 }
 
+
 TSharedRef< ISceneOutliner > FSceneOutlinerModule::CreateSceneOutliner( const FSceneOutlinerInitializationOptions& InitOptions, const FOnSceneOutlinerItemPicked& OnItemPickedDelegate ) const
 {
 	return SNew( SceneOutliner::SSceneOutliner, InitOptions )
diff --git a/Engine/Source/Editor/SceneOutliner/Private/SceneOutlinerSettings.h b/Engine/Source/Editor/SceneOutliner/Private/SceneOutlinerSettings.h
--- a/Engine/Source/Editor/SceneOutliner/Private/SceneOutlinerSettings.h
+++ b/Engine/Source/Editor/SceneOutliner/Private/SceneOutlinerSettings.h
@@ -1,9 +1,5 @@
 // Copyright 1998-2014 Epic Games, Inc. All Rights Reserved.
 
-/*=============================================================================
-	SceneOutlinerSettings.h: Declares the USceneOutlinerSettings class.
-=============================================================================*/
-
 #pragma once
 
 #include "SceneOutlinerSettings.generated.h"
diff --git a/Engine/Source/Editor/SceneOutliner/Public/SceneOutlinerModule.h b/Engine/Source/Editor/SceneOutliner/Public/SceneOutlinerModule.h
--- a/Engine/Source/Editor/SceneOutliner/Public/SceneOutlinerModule.h
+++ b/Engine/Source/Editor/SceneOutliner/Public/SceneOutlinerModule.h
@@ -1,13 +1,8 @@
 // Copyright 1998-2014 Epic Games, Inc. All Rights Reserved.
 
-/*=============================================================================
-	SceneOutlinerModule.h: Declares the FSceneOutlinerModule class.
-=============================================================================*/
-
 #pragma once
 
 #include "ModuleInterface.h"
-
 #include "ISceneOutliner.h"
 #include "SceneOutlinerInitializationOptions.h"
 
@@ -48,10 +43,8 @@ public:
 
 public:
 
-	// Begin IModuleInterface interface
+	// IModuleInterface interface
 
 	virtual void StartupModule() { }
 	virtual void ShutdownModule() { }
-
-	// End IModuleInterface interface
 };
diff --git a/Engine/Source/Editor/TaskBrowser/Private/TaskDataManager.cpp b/Engine/Source/Editor/TaskBrowser/Private/TaskDataManager.cpp
--- a/Engine/Source/Editor/TaskBrowser/Private/TaskDataManager.cpp
+++ b/Engine/Source/Editor/TaskBrowser/Private/TaskDataManager.cpp
@@ -1,13 +1,10 @@
 // Copyright 1998-2014 Epic Games, Inc. All Rights Reserved.
 
-/*=============================================================================
-	TaskDataManager.cpp: Data management layer between task GUI system and database connection
-=============================================================================*/
-
 #include "TaskBrowserPrivatePCH.h"
 #include "TaskDataManager.h"
 #include "TaskDatabase.h"
 
+
 /**
  * FTaskDataManager constructor
  *
@@ -54,7 +51,6 @@ FTaskDataManager::FTaskDataManager( FTaskDataGUIInterface* InGUICallbackObject )
 }
 
 
-
 /** FTaskDataManager destructor */
 FTaskDataManager::~FTaskDataManager()
 {
@@ -70,7 +66,6 @@ FTaskDataManager::~FTaskDataManager()
 }
 
 
-
 /**
  * Called from within LevelTick.cpp after ticking all actors or from
  * the rendering thread (depending on bIsRenderingThreadObject)
@@ -89,6 +84,7 @@ void FTaskDataManager::Tick( float DeltaTime )
 	}
 }
 
+
 TStatId FTaskDataManager::GetStatId() const
 {
 	RETURN_QUICK_DECLARE_CYCLE_STAT(FTaskDataManager, STATGROUP_Tickables);
@@ -117,8 +113,6 @@ void FTaskDataManager::StartMarkingTasksComplete( const TArray< uint32 >& InTask
 }
 
 
-
-
 /** Initiates or re-initiates a connection to the task database server */
 void FTaskDataManager::AttemptConnection()
 {
@@ -130,7 +124,6 @@ void FTaskDataManager::AttemptConnection()
 }
 
 
-
 /** Initiates a disconnection from the task database server */
 void FTaskDataManager::AttemptDisconnection()
 {
@@ -143,7 +136,6 @@ void FTaskDataManager::AttemptDisconnection()
 }
 
 
-
 /** Starts a forced a refresh of task list/description data from the server */
 void FTaskDataManager::ClearTaskDataAndInitiateRefresh()
 {
@@ -157,7 +149,6 @@ void FTaskDataManager::ClearTaskDataAndInitiateRefresh()
 }
 
 
-
 /** Updates our state machine.  Should be called frequently (every Tick is OK) */
 void FTaskDataManager::UpdateTaskDataManager()
 {
@@ -285,7 +276,6 @@ void FTaskDataManager::UpdateTaskDataManager()
 }
 
 
-
 /** We've been disconnected (either voluntarily or otherwise), so clean up state and refresh GUI */
 void FTaskDataManager::CleanUpAfterDisconnect()
 {
@@ -320,7 +310,6 @@ void FTaskDataManager::CleanUpAfterDisconnect()
 }
 
 
-
 /**
  * Called when a response is received from the task database
  *
diff --git a/Engine/Source/Editor/TaskBrowser/Private/TaskDatabase.cpp b/Engine/Source/Editor/TaskBrowser/Private/TaskDatabase.cpp
--- a/Engine/Source/Editor/TaskBrowser/Private/TaskDatabase.cpp
+++ b/Engine/Source/Editor/TaskBrowser/Private/TaskDatabase.cpp
@@ -1,9 +1,5 @@
 // Copyright 1998-2014 Epic Games, Inc. All Rights Reserved.
 
-/*=============================================================================
-	TaskDatabase.cpp: Implements support for communicating with a task database
-=============================================================================*/
-
 #include "TaskBrowserPrivatePCH.h"
 #include "TaskDatabase.h"
 #include "TaskDatabaseThread.h"
diff --git a/Engine/Source/Editor/TaskBrowser/Private/TaskDatabaseThread.cpp b/Engine/Source/Editor/TaskBrowser/Private/TaskDatabaseThread.cpp
--- a/Engine/Source/Editor/TaskBrowser/Private/TaskDatabaseThread.cpp
+++ b/Engine/Source/Editor/TaskBrowser/Private/TaskDatabaseThread.cpp
@@ -1,12 +1,9 @@
 // Copyright 1998-2014 Epic Games, Inc. All Rights Reserved.
 
-/*=============================================================================
-	TaskDatabaseThread.cpp: Implements logic for task database worker thread
-=============================================================================*/
-
 #include "TaskBrowserPrivatePCH.h"
 #include "TaskDatabaseThread.h"
 
+
 /**
  * FTaskDatabaseThreadRunnable constructor
  *
@@ -24,7 +21,6 @@ FTaskDatabaseThreadRunnable::FTaskDatabaseThreadRunnable( FTaskDatabaseProviderI
 }
 
 
-
 /** FTaskDatabaseThreadRunnable destructor */
 FTaskDatabaseThreadRunnable::~FTaskDatabaseThreadRunnable()
 {
@@ -48,7 +44,6 @@ FTaskDatabaseThreadRunnable::~FTaskDatabaseThreadRunnable()
 }
 
 
-
 /**
  * Allows per runnable object initialization. NOTE: This is called in the
  * context of the thread object that aggregates this, not the thread that
@@ -65,8 +60,6 @@ bool FTaskDatabaseThreadRunnable::Init()
 }
 
 
-
-
 /**
  * This is where all per object thread work is done. This is only called
  * if the initialization was successful.
@@ -286,7 +279,6 @@ uint32 FTaskDatabaseThreadRunnable::Run()
 }
 
 
-
 /**
  * This is called if a thread is requested to terminate early
  */
@@ -301,7 +293,6 @@ void FTaskDatabaseThreadRunnable::Stop()
 }
 
 
-
 /**
  * Called in the context of the aggregating thread to perform any cleanup.
  */
@@ -318,4 +309,3 @@ void FTaskDatabaseThreadRunnable::Exit()
 
 	TDLOG( TEXT( "TDThread: Exit() finished" ) );
 }
-
diff --git a/Engine/Source/Editor/TaskBrowser/Private/TestTrackTaskDatabaseProvider.cpp b/Engine/Source/Editor/TaskBrowser/Private/TestTrackTaskDatabaseProvider.cpp
--- a/Engine/Source/Editor/TaskBrowser/Private/TestTrackTaskDatabaseProvider.cpp
+++ b/Engine/Source/Editor/TaskBrowser/Private/TestTrackTaskDatabaseProvider.cpp
@@ -1,12 +1,9 @@
 // Copyright 1998-2014 Epic Games, Inc. All Rights Reserved.
 
-/*=============================================================================
-	TestTrackTaskDatabaseProvider.cpp: Integrates TestTrack Pro task database
-=============================================================================*/
-
 #include "TaskBrowserPrivatePCH.h"
 #include "TestTrackTaskDatabaseProvider.h"
 
+
 #if WITH_TESTTRACK
 
 // TestTrack includes
@@ -29,6 +26,7 @@
 	#endif
 #endif
 
+
 /**
  * Static: Creates a new instance of a TestTrack task database object
  *
@@ -50,7 +48,6 @@ FTestTrackProvider* FTestTrackProvider::CreateTestTrackProvider()
 }
 
 
-
 /**
  * FTestTrackProvider constructor
  */
@@ -61,9 +58,7 @@ FTestTrackProvider::FTestTrackProvider()
 	  ANSIServerEndpointURLString( NULL ),
 	  IDCookie( 0 ),
 	  UserRealName()
-{
-}
-
+{ }
 
 
 /**
@@ -93,7 +88,6 @@ FTestTrackProvider::~FTestTrackProvider()
 }
 
 
-
 /**
  * Initializes this task database
  *
@@ -132,7 +126,6 @@ bool FTestTrackProvider::Init()
 }
 
 
-
 /**
  * Sets the TestTrack server endpoint address using the specified URL string
  *
@@ -159,7 +152,6 @@ void FTestTrackProvider::SetupServerEndpoint( const FString& InServerURL )
 }
 
 
-
 /**
  * Checks the result code and returns whether the operation succeeded.  For failure cases, keeps track
  * of the actual error message so it can be retrieved later.
@@ -220,7 +212,6 @@ bool FTestTrackProvider::VerifyTTPSucceeded( const int32 ResultCode )
 }
 
 
-
 /**
  * Queries the server for a list of databases that the user has access to.  This can be called before
  * the user is logged into the server
@@ -270,7 +261,6 @@ bool FTestTrackProvider::QueryAvailableDatabases( const FString& InServerURL, co
 }
 
 
-
 /**
  * Attempts to connect and login to the specified database
  *
@@ -493,7 +483,6 @@ bool FTestTrackProvider::ConnectToDatabase( const FString& InServerURL, const FS
 }
 
 
-
 /**
  * Logs the user off and disconnects from the database
  *
@@ -520,7 +509,6 @@ bool FTestTrackProvider::DisconnectFromDatabase()
 }
 
 
-
 /**
  * Retrieves a list of filters from the database
  *
@@ -566,7 +554,6 @@ bool FTestTrackProvider::QueryFilters( TArray< FString >& OutFilterNames )
 }
 
 
-
 /**
  * Retrieves a list of tasks from the database that matches the specified filter name
  *
@@ -757,8 +744,6 @@ bool FTestTrackProvider::QueryTasks( const FString& InFilterName, TArray< FTaskD
 }
 
 
-
- 
 /**
  * Retrieves details about a specific task from the database
  *
@@ -840,8 +825,6 @@ bool FTestTrackProvider::QueryTaskDetails( const uint32 InNumber, FTaskDatabaseE
 }
 
 
-
-
 /**
  * Marks the specified task as complete
  *
@@ -1000,5 +983,3 @@ bool FTestTrackProvider::MarkTaskComplete( const uint32 InNumber, const FTaskRes
 
 
 #endif	// WITH_TESTTRACK
-
-
diff --git a/Engine/Source/Editor/TaskBrowser/Private/TestTrackTaskDatabaseProvider.h b/Engine/Source/Editor/TaskBrowser/Private/TestTrackTaskDatabaseProvider.h
--- a/Engine/Source/Editor/TaskBrowser/Private/TestTrackTaskDatabaseProvider.h
+++ b/Engine/Source/Editor/TaskBrowser/Private/TestTrackTaskDatabaseProvider.h
@@ -1,13 +1,10 @@
 // Copyright 1998-2014 Epic Games, Inc. All Rights Reserved.
 
-/*=============================================================================
-	TestTrackTaskDatabaseProvider.h: Integrates TestTrack Pro task database
-=============================================================================*/
-
 #pragma once
 
 #include "TaskDatabaseDefs.h"
 
+
 #if WITH_TESTTRACK
 
 class ttsoapcgi;
@@ -18,19 +15,16 @@ class ttsoapcgi;
 class FTestTrackProvider
 	: public FTaskDatabaseProviderInterface
 {
-
 public:
 
 	/**
 	 * Static: Creates a new instance of a TestTrack task database object
 	 *
 	 * @param	URL for the database server
-	 *
 	 * @return	The new database if successful, otherwise NULL
 	 */
 	static FTestTrackProvider* CreateTestTrackProvider();
 
-
 public:
 
 	/**
@@ -41,12 +35,10 @@ public:
 	 * @param	InUserName			User name string
 	 * @param	InPassword			Password string
 	 * @param	OutDatabaseNames	List of available database names
-	 *
 	 * @return	True if successful
 	 */
 	virtual bool QueryAvailableDatabases( const FString& InServerURL, const FString& InUserName, const FString& InPassword, TArray< FString >& OutDatabaseNames );
 
-
 	/**
 	 * Attempts to connect and login to the specified database
 	 *
@@ -57,12 +49,10 @@ public:
 	 * @param	OutUserRealName			[Out] The real name of our user
 	 * @param	OutResolutionValues		[Out] List of valid fix resolution values
 	 * @param	OutOpenTaskStatusPrefix	[Out] Name of task status for 'Open' tasks
-	 *
 	 * @return	True if successful
 	 */
 	virtual bool ConnectToDatabase( const FString& InServerURL, const FString& InUserName, const FString& InPassword, const FString& InDatabaseName, FString& OutUserRealName, TArray< FString >& OutResolutionValues, FString& OutOpenTaskStatusPrefix );
 
-	
 	/**
 	 * Logs the user off and disconnects from the database
 	 *
@@ -70,50 +60,41 @@ public:
 	 */
 	virtual bool DisconnectFromDatabase();
 
-
 	/**
 	 * Retrieves a list of filter names from the database
 	 *
 	 * @param	OutFilterNames	List of filter names
-	 *
 	 * @return	Returns true if successful
 	 */
 	virtual bool QueryFilters( TArray< FString >& OutFilterNames );
 
-
 	/**
 	 * Retrieves a list of tasks from the database that matches the specified filter name
 	 *
 	 * @param	InFilterName	Filter name to restrict the request by
 	 * @param	OutTaskList		List of downloaded tasks
-	 *
 	 * @return	Returns true if successful
 	 */
 	virtual bool QueryTasks( const FString& InFilterName, TArray< FTaskDatabaseEntry >& OutTaskList );
 
-
 	/**
 	 * Retrieves details about a specific task from the database
 	 *
 	 * @param	InNumber		Task number
 	 * @param	OutDetails		[Out] Details for the requested task
-	 *
 	 * @return	Returns true if successful
 	 */
 	virtual bool QueryTaskDetails( const uint32 InNumber, FTaskDatabaseEntryDetails& OutDetails );
 
-
 	/**
 	 * Marks the specified task as complete
 	 *
 	 * @param	InNumber		Task number
 	 * @param	InResolutionData	Resolution data for this task
-	 *
 	 * @return	Returns true if successful
 	 */
 	virtual bool MarkTaskComplete( const uint32 InNumber, const FTaskResolutionData& InResolutionData );
 
-
 	/**
 	 * Returns a detailed error message for the last reported error.  This message may be displayed to the user.
 	 *
@@ -124,7 +105,6 @@ public:
 		return LastErrorMessage;
 	}
 
-
 	/**
 	 * Returns true if the last error was caused by a disconnection from the server
 	 *
@@ -135,8 +115,6 @@ public:
 		return bLastErrorWasCausedByDisconnection;
 	}
 
-
-
 protected:
 
 	/**
@@ -144,7 +122,6 @@ protected:
 	 */
 	FTestTrackProvider();
 
-
 public:
 
 	/**
@@ -152,7 +129,6 @@ public:
 	 */
 	virtual ~FTestTrackProvider();
 
-
 protected:
 
 	/**
@@ -162,7 +138,6 @@ protected:
 	 */
 	bool Init();
 
-	
 	/**
 	 * Sets the TestTrack server endpoint address using the specified URL string
 	 *
@@ -170,18 +145,15 @@ protected:
 	 */
 	void SetupServerEndpoint( const FString& InServerURL );
 
-
 	/**
 	 * Checks the result code and returns whether the operation succeeded.  For failure cases, keeps track
 	 * of the actual error message so it can be retrieved later.
 	 *
 	 * @param	ResultCode	The result code to test
-	 *
 	 * @return	True if the result code reported success
 	 */
 	bool VerifyTTPSucceeded( const int32 ResultCode );
 
-
 protected:
 
 	/** SOAP interface to TestTrack web server */
diff --git a/Engine/Source/Editor/TaskBrowser/Public/TaskDataManager.h b/Engine/Source/Editor/TaskBrowser/Public/TaskDataManager.h
--- a/Engine/Source/Editor/TaskBrowser/Public/TaskDataManager.h
+++ b/Engine/Source/Editor/TaskBrowser/Public/TaskDataManager.h
@@ -1,13 +1,10 @@
 // Copyright 1998-2014 Epic Games, Inc. All Rights Reserved.
 
-/*=============================================================================
-	TaskDataManager.h: Data management layer between task GUI system and database connection
-=============================================================================*/
-
 #pragma once
 
 #include "TaskDatabaseDefs.h"
 
+
 /** GUI refresh options for the task browser */
 namespace ETaskBrowserGUIRefreshOptions
 {
@@ -31,14 +28,11 @@ namespace ETaskBrowserGUIRefreshOptions
 }
 
 
-
-
 /**
  * Classes derived from this implement a GUI front end for the task data manager
  */
 class FTaskDataGUIInterface
 {
-
 public:
 
 	/**
@@ -47,12 +41,9 @@ public:
 	 * @param	Options		Bitfield that contains the types of GUI elements to refresh
 	 */
 	virtual void Callback_RefreshGUI( const ETaskBrowserGUIRefreshOptions::Type Options ) = 0;
-
 };
 
 
-
-
 /** Status field for the task data manager */
 namespace ETaskDataManagerStatus
 {
@@ -78,11 +69,9 @@ namespace ETaskDataManagerStatus
 
 		/** Disconnecting */
 		Disconnecting,
-
 		
 		// ...
 
-
 		/** Querying filters */
 		QueryingFilters,
 
@@ -94,19 +83,15 @@ namespace ETaskDataManagerStatus
 
 		/** Marking task as fixed */
 		MarkingTaskComplete,
-
 	};
 }
 
 
-
-
 /**
  * Connection settings and credentials for logging into the task database
  */
 class FTaskDataManagerConnectionSettings
 {
-
 public:
 
 	/** Server name */
@@ -124,7 +109,6 @@ public:
 	/** Project name */
 	FString ProjectName;
 
-
 public:
 
 	/** FTaskDataManagerConnectionSettings constructor */
@@ -134,14 +118,10 @@ public:
 		  UserName(),
 		  Password(),
 		  ProjectName()
-	{
-	}
-
+	{ }
 };
 
 
-
-
 /**
  * Data management layer between task GUI system and database connection
  */
@@ -149,7 +129,6 @@ class FTaskDataManager
 	: public FTaskDatabaseListener,		// Listens for task database responses
 	  public FTickableEditorObject			// Tickable, so we can update the task database
 {
-
 public:
 
 	/**
@@ -159,9 +138,13 @@ public:
 	 */
 	FTaskDataManager( FTaskDataGUIInterface* InGUICallbackObject );
 
-	/** FTaskDataManager destructor */
+	/**
+	 * Virtual destructor.
+	 */
 	virtual ~FTaskDataManager();
 
+public:
+
 	/**
 	 * Used to determine whether an object is ready to be ticked. This is 
 	 * required for example for all UObject derived classes as they might be
@@ -174,7 +157,6 @@ public:
 		return true;
 	}
 
-
 	/**
 	 * Called from within LevelTick.cpp after ticking all actors or from
 	 * the rendering thread (depending on bIsRenderingThreadObject)
@@ -185,7 +167,6 @@ public:
 
 	virtual TStatId GetStatId() const OVERRIDE;
 
-
 	/**
 	 * Sets the connection settings used to connect and login to the task database server
 	 *
@@ -196,7 +177,6 @@ public:
 		ConnectionSettings = InConnectionSettings;
 	}
 
-
 	/**
 	 * Returns the user's real name (retrieved from the server, if available)
 	 *
@@ -218,7 +198,6 @@ public:
 		return ResolutionValues;
 	}
 
-
 	/**
 	 * Returns the string name for 'Open' task status
 	 *
@@ -229,7 +208,6 @@ public:
 		return OpenTaskStatusPrefix;
 	}
 
-
 	/**
 	 * Returns the current connection status of the task data manager
 	 *
@@ -240,7 +218,6 @@ public:
 		return ConnectionStatus;
 	}
 
-
 	/**
 	 * Returns the current status to display in the GUI
 	 *
@@ -260,7 +237,6 @@ public:
 		return GUIStatus;
 	}
 
-
 	/**
 	 * Changes the active filter
 	 *
@@ -275,7 +251,6 @@ public:
 		}
 	}
 
-
 	/**
 	 * Returns the active filter name
 	 *
@@ -286,7 +261,6 @@ public:
 		return ActiveFilterName;
 	}
 
-
 	/**
 	 * Sets the 'focused task number' (the task we should retrieve details about)
 	 *
@@ -301,7 +275,6 @@ public:
 		}
 	}
 
-
 	/**
 	 * Returns the current focused task number
 	 *
@@ -312,7 +285,6 @@ public:
 		return FocusedTaskNumber;
 	}
 
-
 	/**
 	 * Starts process of marking the specified task numbers as completed
 	 *
@@ -322,7 +294,6 @@ public:
 	void StartMarkingTasksComplete( const TArray< uint32 >& InTaskNumbers,
 									const FTaskResolutionData& InResolutionData );
 
-
 	/**
 	 * Returns the list of cached filter names
 	 *
@@ -333,7 +304,6 @@ public:
 		return CachedFilterNames;
 	}
 
-
 	/**
 	 * Returns the list of cached task entries
 	 *
@@ -344,8 +314,6 @@ public:
 		return CachedTaskArray;
 	}
 
-
-
 	/**
 	 * Searches for details about the specified task number in our cache and returns it if found
 	 *
@@ -356,28 +324,21 @@ public:
 		return CachedTaskDetailsMap.Find( InTaskNumber );
 	}
 
-
-
 	/** Initiates or re-initiates a connection to the task database server */
 	void AttemptConnection();
 
-
 	/** Initiates a disconnection from the task database server */
 	void AttemptDisconnection();
 
-
 	/** Starts a forced a refresh of task list/description data from the server */
 	void ClearTaskDataAndInitiateRefresh();
 
-
 	/** Updates our state machine.  Should be called frequently (every Tick is OK) */
 	void UpdateTaskDataManager();
 
-
 	/** We've been disconnected (either voluntarily or otherwise), so clean up state and refresh GUI */
 	void CleanUpAfterDisconnect();
 
-
 	/**
 	 * Called when a response is received from the task database
 	 *
@@ -386,7 +347,6 @@ public:
 	 */
 	virtual void OnTaskDatabaseRequestCompleted( const FTaskDatabaseResponse* InGenericResponse );
 
-
 private:
 
 	/** Callback object for GUI updates */
@@ -431,7 +391,7 @@ private:
 	uint32 CurrentlyProcessingTaskNumber;
 
 	/** Array of cached filters from server */
-	TArray< FString > CachedFilterNames;
+	TArray<FString> CachedFilterNames;
 
 	/** Array of cached tasks from server */
 	TArray< FTaskDatabaseEntry > CachedTaskArray;
@@ -440,6 +400,5 @@ private:
 	FString CachedLastTaskArrayFilterName;
 
 	/** Map of task number to cached task details */
-	TMap< int32, FTaskDatabaseEntryDetails > CachedTaskDetailsMap;
-
+	TMap<int32, FTaskDatabaseEntryDetails> CachedTaskDetailsMap;
 };
diff --git a/Engine/Source/Editor/TaskBrowser/Public/TaskDatabase.h b/Engine/Source/Editor/TaskBrowser/Public/TaskDatabase.h
--- a/Engine/Source/Editor/TaskBrowser/Public/TaskDatabase.h
+++ b/Engine/Source/Editor/TaskBrowser/Public/TaskDatabase.h
@@ -42,13 +42,11 @@ namespace ETaskDatabaseRequestType
 }
 
 
-
 /**
  * Generic response to a task database request
  */
 class FTaskDatabaseResponse
 {
-
 public:
 
 	/** The request type that triggered this response */
@@ -62,104 +60,85 @@ public:
 
 	/** Detailed error message, in the case of failure */
 	FString ErrorMessage;
-
 };
 
 
-
 /**
  * Response to 'ConnectToDatabase' task database request
  */
 class FTaskDatabaseResponse_ConnectToDatabase
 	: public FTaskDatabaseResponse
 {
-
 public:
 
 	/** The connected user's real name */
 	FString UserRealName;
 
 	/** List of valid bug 'resolution' string values */
-	TArray< FString > ResolutionValues;
+	TArray<FString> ResolutionValues;
 
 	/** Status string name for 'Open' tasks (*not* completed) */
 	FString OpenTaskStatusPrefix;
-
 };
 
 
-
 /**
  * Response to 'QueryAvailableDatabases' task database request
  */
 class FTaskDatabaseResponse_QueryAvailableDatabases
 	: public FTaskDatabaseResponse
 {
-
 public:
 
 	/** Array of available database names */
-	TArray< FString > DatabaseNames;
-
+	TArray<FString> DatabaseNames;
 };
 
 
-
 /**
  * Response to 'QueryFilters' task database request
  */
 class FTaskDatabaseResponse_QueryFilters
 	: public FTaskDatabaseResponse
 {
-
 public:
 
 	/** Array of filter names */
-	TArray< FString > FilterNames;
-
+	TArray<FString> FilterNames;
 };
 
 
-
 /**
  * Response to 'QueryTasks' task database request
  */
 class FTaskDatabaseResponse_QueryTasks
 	: public FTaskDatabaseResponse
 {
-
 public:
 
 	/** Array of task entries that match the query */
-	TArray< FTaskDatabaseEntry > TaskEntries;
-
+	TArray<FTaskDatabaseEntry> TaskEntries;
 };
 
 
-
 /**
  * Response to 'QueryTaskDetails' task database request
  */
 class FTaskDatabaseResponse_QueryTaskDetails
 	: public FTaskDatabaseResponse
 {
-
 public:
 
 	/** Details about the requested task */
 	FTaskDatabaseEntryDetails TaskEntryDetails;
-
 };
 
 
-
-
 /**
  * Task database system, public API
  */
 namespace TaskDatabaseSystem
 {
-
 	/**
 	 * Initialize the task database system.  This must be called before anything else.
 	 *
@@ -167,11 +146,9 @@ namespace TaskDatabaseSystem
 	 */
 	bool Init();
 
-
 	/** Clean up the task database system.  Call this after before exiting the app. */
 	void Destroy();
 
-
 	/**
 	 * Returns true if the task database system was initialized successfully
 	 *
@@ -179,7 +156,6 @@ namespace TaskDatabaseSystem
 	 */
 	bool IsInitialized();
 
-
 	/**
 	 * Returns true if we're currently waiting for a response from the server
 	 *
@@ -187,7 +163,6 @@ namespace TaskDatabaseSystem
 	 */
 	bool IsRequestInProgress();
 
-
 	/**
 	 * Returns true if we're currently connected to a database
 	 *
@@ -195,12 +170,10 @@ namespace TaskDatabaseSystem
 	 */
 	bool IsConnected();
 
-
 	/** Updates the task database and checks for any completed requests, firing callbacks if needed.  This
 	    should usually be called with every engine Tick. */
 	void Update();
 
-
 	/**
 	 * Adds the specified object as a 'listener' for completed task database requests.  Remember to remove
 	 * the object using UnregisterListener() before it's deleted!
@@ -209,7 +182,6 @@ namespace TaskDatabaseSystem
 	 */
 	void RegisterListener( FTaskDatabaseListener* Listener );
 
-	
 	/**
 	 * Removes the specified object from the list of 'listeners'
 	 *
@@ -217,7 +189,6 @@ namespace TaskDatabaseSystem
 	 */
 	void UnregisterListener( FTaskDatabaseListener* Listener );
 
-
 	/**
 	 * Initiates asynchronous query for available task databases
 	 *
@@ -225,12 +196,10 @@ namespace TaskDatabaseSystem
 	 * @param	InServerPort		Server port
 	 * @param	InLoginUserName		User name
 	 * @param	InLoginPassword		Password
-	 *
 	 * @return	True if async requested was kicked off successfully
 	 */
 	bool QueryAvailableDatabases_Async( const FString& InServerName, const uint32 InServerPort, const FString& InLoginUserName, const FString& InLoginPassword );
 
-
 	/**
 	 * Asynchronously connects to the task database and attempts to login to a specific database
 	 *
@@ -239,12 +208,10 @@ namespace TaskDatabaseSystem
 	 * @param	InLoginUserName		User name
 	 * @param	InLoginPassword		Password
 	 * @param	InDatabaeName		The database to login to
-	 *
 	 * @return	True if async requested was kicked off successfully
 	 */
 	bool ConnectToDatabase_Async( const FString& InServerName, const uint32 InServerPort, const FString& InLoginUserName, const FString& InLoginPassword, const FString& InDatabaseName );
 
-
 	/**
 	 * Disconnects from the database asynchronously
 	 *
@@ -252,7 +219,6 @@ namespace TaskDatabaseSystem
 	 */
 	bool DisconnectFromDatabase_Async();
 
-
 	/**
 	 * Asynchronously queries a list of available database filters
 	 *
@@ -260,35 +226,28 @@ namespace TaskDatabaseSystem
 	 */
 	bool QueryFilters_Async();
 
-
 	/**
 	 * Asynchronously queries a list of task entries for the specified filter name
 	 *
 	 * @param	InFilterName		Name of the filter (can be empty string)
-	 *
 	 * @return	True if async requested was kicked off successfully
 	 */
 	bool QueryTasks_Async( const FString& InFilterName );
 
-
 	/**
 	 * Asynchronously queries details about a single specific task entry
 	 *
 	 * @param	InTaskNumber		The database task number for the task entry
-	 *
 	 * @return	True if async requested was kicked off successfully
 	 */
 	bool QueryTaskDetails_Async( const uint32 InTaskNumber );
 
-
 	/**
 	 * Attempts to asynchronously mark the specified task as completed
 	 *
 	 * @param	InTaskNumber		The database task number to mark complete
 	 * @param	InResolutionData	Information about how the task should be resolved
-	 *
 	 * @return	True if async requested was kicked off successfully
 	 */
 	bool MarkTaskComplete_Async( const uint32 InTaskNumber, const FTaskResolutionData& InResolutionData );
-
 };
diff --git a/Engine/Source/Editor/TaskBrowser/Public/TaskDatabaseDefs.h b/Engine/Source/Editor/TaskBrowser/Public/TaskDatabaseDefs.h
--- a/Engine/Source/Editor/TaskBrowser/Public/TaskDatabaseDefs.h
+++ b/Engine/Source/Editor/TaskBrowser/Public/TaskDatabaseDefs.h
@@ -1,22 +1,13 @@
 // Copyright 1998-2014 Epic Games, Inc. All Rights Reserved.
 
-/*=============================================================================
-	TaskDatabaseDefs.h: Task database common definitions and types
-=============================================================================*/
-
 #pragma once
 
-#ifndef __TaskDatabaseDefs_h__
-#define __TaskDatabaseDefs_h__
 
 DECLARE_LOG_CATEGORY_EXTERN(LogTaskDatabaseDefs, Log, All);
 
-
-
 class FTaskDatabaseResponse;
 
 
-
 /** Define LOG_TASK_DATABASE to 1 to enable debug logging throughout the task database system */
 #define LOG_TASK_DATABASE 0
 
@@ -27,13 +18,11 @@ class FTaskDatabaseResponse;
 #endif
 
 
-
 /**
- * Base information about a single task in the database
+ * Base information about a single task in the database.
  */
 class FTaskDatabaseEntryBase
 {
-
 public:
 
 	/** Task number */
@@ -51,76 +40,70 @@ public:
 	/** Who entered this task */
 	FString CreatedBy;
 
-
 public:
 
-	/** FTaskDatabaseEntry constructor */
+	/**
+	 * Default constructor.
+	 */
 	FTaskDatabaseEntryBase()
 		: Number( 0 ),
 		Priority(),
 		Summary(),
 		Status(),
 		CreatedBy()
-	{
-	}
+	{ }
 };
 
 
 /**
- * Information about a single task in the database
+ * Information about a single task in the database.
  */
 class FTaskDatabaseEntry : public FTaskDatabaseEntryBase
 {
-	
 public:
 
 	/** Assigned To */
 	FString AssignedTo;
 
-
 public:
 
-	/** FTaskDatabaseEntry constructor */
+	/**
+	 * Default constructor.
+	 */
 	FTaskDatabaseEntry()
 		: FTaskDatabaseEntryBase(),
 		  AssignedTo()
-	{
-	}
+	{ }
 };
 
 
-
 /**
- * Details about a task entry
+ * Details about a task entry.
  */
 class FTaskDatabaseEntryDetails : public FTaskDatabaseEntryBase
 {
-
 public:
 
 	/** Description */
 	FString Description;
 
-
 public:
 
-	/** FTaskDatabaseEntryDetails constructor */
+	/**
+	 * Default constructor.
+	 */
 	FTaskDatabaseEntryDetails()
 		: FTaskDatabaseEntryBase(),
 		  Description()
-	{
-	}
-
+	{ }
 };
 
 
-
 /**
  * Contains the user-entered data needed to resolve a task
  */
 class FTaskResolutionData
 {
-
 public:
 
 	/** Resolution type */
@@ -135,17 +118,17 @@ public:
 	/** Hours to complete */
 	double HoursToComplete;
 
-
 public:
 
-	/** FTaskResolutionData constructor */
+	/**
+	 * Default constructor.
+	 */
 	FTaskResolutionData()
 		: ResolutionType(),
 		  Comments(),
 		  ChangelistNumber( 0 ),
 		  HoursToComplete( -1.0 )		  
-	{
-	}
+	{ }
 
 	bool IsValid()
 	{
@@ -161,29 +144,24 @@ public:
 };
 
 
-
 /**
  * Callback interface for objects that 'listen' to task manager events
  */
 class FTaskDatabaseListener
 {
-
 public:
 
 	/** Called when a response is received from the task database.  This is a polymorphic object that should
 		be cast to the appropriate response type for the request. */
 	virtual void OnTaskDatabaseRequestCompleted( const FTaskDatabaseResponse* InGenericResponse ) = 0;
-
 };
 
 
-
 /**
  * Interface class for implementing a task database provider
  */
 class FTaskDatabaseProviderInterface
 {
-
 public:
 
 	/** FTaskDatabaseProviderInterface destructor */
@@ -197,12 +175,10 @@ public:
 	 * @param	InUserName			User name string
 	 * @param	InPassword			Password string
 	 * @param	OutDatabaseNames	List of available database names
-	 *
 	 * @return	True if successful
 	 */
 	virtual bool QueryAvailableDatabases( const FString& InServerURL, const FString& InUserName, const FString& InPassword, TArray< FString >& OutDatabaseNames ) = 0;
 
-
 	/**
 	 * Attempts to connect and login to the specified database
 	 *
@@ -213,12 +189,10 @@ public:
 	 * @param	OutUserRealName			[Out] The real name of our user
 	 * @param	OutResolutionValues		[Out] List of valid fix resolution values
 	 * @param	OutOpenTaskStatusPrefix	[Out] Name of task status for 'Open' tasks
-	 *
 	 * @return	True if successful
 	 */
 	virtual bool ConnectToDatabase( const FString& InServerURL, const FString& InUserName, const FString& InPassword, const FString& InDatabaseName, FString& OutUserRealName, TArray< FString >& OutResolutionValues, FString& OutOpenTaskStatusPrefix ) = 0;
 
-
 	/**
 	 * Logs the user off and disconnects from the database
 	 *
@@ -226,50 +200,41 @@ public:
 	 */
 	virtual bool DisconnectFromDatabase() = 0;
 
-
 	/**
 	 * Retrieves a list of filter names from the database
 	 *
 	 * @param	OutFilterNames	List of filter names
-	 *
 	 * @return	Returns true if successful
 	 */
 	virtual bool QueryFilters( TArray< FString >& OutFilterNames ) = 0;
 
-
 	/**
 	 * Retrieves a list of tasks from the database that matches the specified filter name
 	 *
 	 * @param	InFilterName	Filter name to restrict the request by
 	 * @param	OutTaskList		List of downloaded tasks
-	 *
 	 * @return	Returns true if successful
 	 */
 	virtual bool QueryTasks( const FString& InFilterName, TArray< FTaskDatabaseEntry >& OutTaskList ) = 0;
 
-
 	/**
 	 * Retrieves details about a specific task from the database
 	 *
 	 * @param	InNumber		Task number
 	 * @param	OutDetails		[Out] Details for the requested task
-	 *
 	 * @return	Returns true if successful
 	 */
 	virtual bool QueryTaskDetails( const uint32 InNumber, FTaskDatabaseEntryDetails& OutDetails ) = 0;
 
-
 	/**
 	 * Marks the specified task as complete
 	 *
 	 * @param	InNumber			Task number
 	 * @param	InResolutionData	Resolution data for this task
-	 *
 	 * @return	Returns true if successful
 	 */
 	virtual bool MarkTaskComplete( const uint32 InNumber, const FTaskResolutionData& InResolutionData ) = 0;
 
-
 	/**
 	 * Returns a detailed error message for the last reported error.  This message may be displayed to the user.
 	 *
@@ -277,18 +242,10 @@ public:
 	 */
 	virtual const FString& GetLastErrorMessage() const = 0;
 
-
 	/**
 	 * Returns true if the last error was caused by a disconnection from the server
 	 *
 	 * @return	True if the last error was a disconnection error
 	 */
 	virtual const bool LastErrorWasCausedByDisconnection() const = 0;
-
 };
-
-
-
-#endif // __TaskDatabaseDefs_h__
-
-
diff --git a/Engine/Source/Editor/TextureEditor/Classes/TextureEditorSettings.h b/Engine/Source/Editor/TextureEditor/Classes/TextureEditorSettings.h
--- a/Engine/Source/Editor/TextureEditor/Classes/TextureEditorSettings.h
+++ b/Engine/Source/Editor/TextureEditor/Classes/TextureEditorSettings.h
@@ -1,9 +1,5 @@
 // Copyright 1998-2014 Epic Games, Inc. All Rights Reserved.
 
-/*=============================================================================
-	TextureEditorSettings.h: Declares the UTextureEditorSettings class.
-=============================================================================*/
-
 #pragma once
 
 
@@ -33,53 +29,37 @@ class TEXTUREEDITOR_API UTextureEditorSettings
 
 public:
 
-	/**
-	 * The type of background to draw in the texture editor view port.
-	 */
+	/** The type of background to draw in the texture editor view port. */
 	UPROPERTY(config)
 	TEnumAsByte<ETextureEditorBackgrounds> Background;
 
-	/**
-	 * Background and foreground color used by Texture preview view ports.
-	 */
+	/** Background and foreground color used by Texture preview view ports. */
 	UPROPERTY(config, EditAnywhere, Category=Background)
 	FColor BackgroundColor;
 
-	/**
-	 * The first color of the checkered background.
-	 */
+	/** The first color of the checkered background. */
 	UPROPERTY(config, EditAnywhere, Category=Background)
 	FColor CheckerColorOne;
 
-	/**
-	 * The second color of the checkered background.
-	 */
+	/** The second color of the checkered background. */
 	UPROPERTY(config, EditAnywhere, Category=Background)
 	FColor CheckerColorTwo;
 
-	/**
-	 * The size of the checkered background tiles.
-	 */
+	/** The size of the checkered background tiles. */
 	UPROPERTY(config, EditAnywhere, Category=Background)
 	int32 CheckerSize;
 
 public:
 
-	/**
-	 * Whether the texture should scale to fit the view port.
-	 */
+	/** Whether the texture should scale to fit the view port. */
 	UPROPERTY(config)
 	bool FitToViewport;
 
-	/**
-	 * Color to use for the texture border, if enabled.
-	 */
+	/** Color to use for the texture border, if enabled. */
 	UPROPERTY(config, EditAnywhere, Category=TextureBorder)
 	FColor TextureBorderColor;
 
-	/**
-	 * If true, displays a border around the texture.
-	 */
+	/** If true, displays a border around the texture. */
 	UPROPERTY(config)
 	bool TextureBorderEnabled;
 };
diff --git a/Engine/Source/Editor/TextureEditor/Private/Menus/TextureEditorViewOptionsMenu.h b/Engine/Source/Editor/TextureEditor/Private/Menus/TextureEditorViewOptionsMenu.h
--- a/Engine/Source/Editor/TextureEditor/Private/Menus/TextureEditorViewOptionsMenu.h
+++ b/Engine/Source/Editor/TextureEditor/Private/Menus/TextureEditorViewOptionsMenu.h
@@ -1,9 +1,5 @@
 // Copyright 1998-2014 Epic Games, Inc. All Rights Reserved.
 
-/*=============================================================================
-	TextureEditorViewOptionsMenu.h: Declares the FTextureEditorViewOptionsMenu class.
-=============================================================================*/
-
 #pragma once
 
 
@@ -20,7 +16,7 @@ public:
 	/**
 	 * Creates the menu.
 	 *
-	 * @param MenuBuilder - The builder for the menu that owns this menu.
+	 * @param MenuBuilder The builder for the menu that owns this menu.
 	 */
 	static void MakeMenu( FMenuBuilder& MenuBuilder )
 	{
@@ -58,7 +54,7 @@ protected:
 	/**
 	 * Creates the 'Background' sub-menu.
 	 *
-	 * @param MenuBuilder - The builder for the menu that owns this menu.
+	 * @param MenuBuilder The builder for the menu that owns this menu.
 	 */
 	static void GenerateBackgroundMenuContent( FMenuBuilder& MenuBuilder )
 	{
diff --git a/Engine/Source/Editor/TextureEditor/Private/Models/TextureEditorViewportClient.cpp b/Engine/Source/Editor/TextureEditor/Private/Models/TextureEditorViewportClient.cpp
--- a/Engine/Source/Editor/TextureEditor/Private/Models/TextureEditorViewportClient.cpp
+++ b/Engine/Source/Editor/TextureEditor/Private/Models/TextureEditorViewportClient.cpp
@@ -1,9 +1,5 @@
 // Copyright 1998-2014 Epic Games, Inc. All Rights Reserved.
 
-/*=============================================================================
-	TextureEditorViewportClient.cpp: Implements the FTextureEditorViewportClient class.
-=============================================================================*/
-
 #include "TextureEditorPrivatePCH.h"
 
 
diff --git a/Engine/Source/Editor/TextureEditor/Private/Models/TextureEditorViewportClient.h b/Engine/Source/Editor/TextureEditor/Private/Models/TextureEditorViewportClient.h
--- a/Engine/Source/Editor/TextureEditor/Private/Models/TextureEditorViewportClient.h
+++ b/Engine/Source/Editor/TextureEditor/Private/Models/TextureEditorViewportClient.h
@@ -1,9 +1,5 @@
 // Copyright 1998-2014 Epic Games, Inc. All Rights Reserved.
 
-/*=============================================================================
-	TextureEditorViewportClient.h: Declares the FTextureEditorViewportClient class.
-=============================================================================*/
-
 #pragma once
 
 
diff --git a/Engine/Source/Editor/TextureEditor/Private/TextureEditorClasses.cpp b/Engine/Source/Editor/TextureEditor/Private/TextureEditorClasses.cpp
--- a/Engine/Source/Editor/TextureEditor/Private/TextureEditorClasses.cpp
+++ b/Engine/Source/Editor/TextureEditor/Private/TextureEditorClasses.cpp
@@ -1,13 +1,8 @@
 // Copyright 1998-2014 Epic Games, Inc. All Rights Reserved.
 
-/*=============================================================================
-	TextureEditorClasses.cpp: Implements the module's script classes.
-=============================================================================*/
-
 #include "TextureEditorPrivatePCH.h"
 
 
-
 UTextureEditorSettings::UTextureEditorSettings( const class FPostConstructInitializeProperties& PCIP )
 	: Super(PCIP)
 	, Background(TextureEditorBackground_SolidColor)
diff --git a/Engine/Source/Editor/TextureEditor/Private/TextureEditorModule.cpp b/Engine/Source/Editor/TextureEditor/Private/TextureEditorModule.cpp
--- a/Engine/Source/Editor/TextureEditor/Private/TextureEditorModule.cpp
+++ b/Engine/Source/Editor/TextureEditor/Private/TextureEditorModule.cpp
@@ -19,9 +19,9 @@ class FTextureEditorModule
 {
 public:
 
-	// Begin ITextureEditorModule interface
+	// ITextureEditorModule interface
 
-	virtual TSharedRef<ITextureEditorToolkit> CreateTextureEditor( const EToolkitMode::Type Mode, const TSharedPtr< IToolkitHost >& InitToolkitHost, UTexture* Texture ) OVERRIDE
+	virtual TSharedRef<ITextureEditorToolkit> CreateTextureEditor( const EToolkitMode::Type Mode, const TSharedPtr< IToolkitHost >& InitToolkitHost, UTexture* Texture ) override
 	{
 		TSharedRef<FTextureEditorToolkit> NewTextureEditor(new FTextureEditorToolkit());
 		NewTextureEditor->InitTextureEditor(Mode, InitToolkitHost, Texture);
@@ -39,13 +39,11 @@ public:
 		return ToolBarExtensibilityManager;
 	}
 
-	// End ITextureEditorModule interface
-
 public:
 
-	// Begin IModuleInterface interface
+	// IModuleInterface interface
 
-	virtual void StartupModule( ) OVERRIDE
+	virtual void StartupModule( ) override
 	{
 		// register menu extensions
 		MenuExtensibilityManager = MakeShareable(new FExtensibilityManager);
@@ -64,7 +62,7 @@ public:
 		}
 	}
 
-	virtual void ShutdownModule( ) OVERRIDE
+	virtual void ShutdownModule( ) override
 	{
 		// unregister settings
 		ISettingsModule* SettingsModule = ISettingsModule::Get();
@@ -79,8 +77,6 @@ public:
 		ToolBarExtensibilityManager.Reset();
 	}
 
-	// End IModuleInterface interface
-
 private:
 
 	// Holds the menu extensibility manager.
diff --git a/Engine/Source/Editor/TextureEditor/Private/TextureEditorPrivatePCH.h b/Engine/Source/Editor/TextureEditor/Private/TextureEditorPrivatePCH.h
--- a/Engine/Source/Editor/TextureEditor/Private/TextureEditorPrivatePCH.h
+++ b/Engine/Source/Editor/TextureEditor/Private/TextureEditorPrivatePCH.h
@@ -1,17 +1,11 @@
 // Copyright 1998-2014 Epic Games, Inc. All Rights Reserved.
 
-/*=============================================================================
-	TextureEditorPrivatePCH.h: Pre-compiled header file for the TextureEditor module.
-=============================================================================*/
+#pragma once
 
-#ifndef TEXTURE_EDITOR_PRIVATEPCH_H
-#define TEXTURE_EDITOR_PRIVATEPCH_H
+#include "TextureEditor.h"
 
 
-#include "../Public/TextureEditor.h"
-
-
-/* Dependencies
+/* Private dependencies
  *****************************************************************************/
 
 #include "CubemapUnwrapUtils.h"
@@ -36,6 +30,7 @@ const double MaxZoom = 16.0;
 const double MinZoom = 0.01;
 const double ZoomStep = 0.1;
 
+
 #include "TextureEditorCommands.h"
 #include "TextureEditorSettings.h"
 #include "TextureEditorViewOptionsMenu.h"
@@ -43,6 +38,3 @@ const double ZoomStep = 0.1;
 #include "STextureEditorViewport.h"
 #include "TextureEditorViewportClient.h"
 #include "TextureEditorToolkit.h"
-
-
-#endif
diff --git a/Engine/Source/Editor/TextureEditor/Private/Widgets/STextureEditorViewport.cpp b/Engine/Source/Editor/TextureEditor/Private/Widgets/STextureEditorViewport.cpp
--- a/Engine/Source/Editor/TextureEditor/Private/Widgets/STextureEditorViewport.cpp
+++ b/Engine/Source/Editor/TextureEditor/Private/Widgets/STextureEditorViewport.cpp
@@ -1,13 +1,9 @@
 // Copyright 1998-2014 Epic Games, Inc. All Rights Reserved.
 
-/*=============================================================================
-	STextureEditorViewport.cpp: Implements the STextureEditorViewport class.
-=============================================================================*/
-
 #include "TextureEditorPrivatePCH.h"
 
-#define LOCTEXT_NAMESPACE "STextureEditorViewport"
 
+#define LOCTEXT_NAMESPACE "STextureEditorViewport"
 
 // Specifies the maximum allowed exposure bias.
 const int32 MaxExposure = 10;
diff --git a/Engine/Source/Editor/TextureEditor/Private/Widgets/STextureEditorViewport.h b/Engine/Source/Editor/TextureEditor/Private/Widgets/STextureEditorViewport.h
--- a/Engine/Source/Editor/TextureEditor/Private/Widgets/STextureEditorViewport.h
+++ b/Engine/Source/Editor/TextureEditor/Private/Widgets/STextureEditorViewport.h
@@ -1,9 +1,5 @@
 // Copyright 1998-2014 Epic Games, Inc. All Rights Reserved.
 
-/*=============================================================================
-	STextureEditorViewport.h: Declares the STextureEditorViewport class.
-=============================================================================*/
-
 #pragma once
 
 
@@ -27,7 +23,7 @@ public:
 	/**
 	 * Constructs the widget.
 	 *
-	 * @param InArgs - The construction arguments.
+	 * @param InArgs The construction arguments.
 	 */
 	void Construct( const FArguments& InArgs, const TSharedRef<ITextureEditorToolkit>& InToolkit );
 	
@@ -59,15 +55,17 @@ public:
 
 public:
 
-	// Begin SWidget overrides
+	// SWidget overrides
 
 	virtual void Tick( const FGeometry& AllottedGeometry, const double InCurrentTime, const float InDeltaTime ) OVERRIDE;
 
-	// End SWidget overrides
-
 protected:
 
-	/** Returns a string representation of the currently displayed textures resolution */
+	/**
+	 * Gets the displayed textures resolution as a string.
+	 *
+	 * @return Texture resolution string.
+	 */
 	FString GetDisplayedResolution() const;
 
 private:
diff --git a/Engine/Source/Editor/TextureEditor/Public/Interfaces/ITextureEditorModule.h b/Engine/Source/Editor/TextureEditor/Public/Interfaces/ITextureEditorModule.h
--- a/Engine/Source/Editor/TextureEditor/Public/Interfaces/ITextureEditorModule.h
+++ b/Engine/Source/Editor/TextureEditor/Public/Interfaces/ITextureEditorModule.h
@@ -1,9 +1,5 @@
 // Copyright 1998-2014 Epic Games, Inc. All Rights Reserved.
 
-/*=============================================================================
-	ITextureEditorModule.h: Declares the ITextureEditorModule interface.
-=============================================================================*/
-
 #pragma once
 
 
@@ -20,9 +16,9 @@ public:
 	/**
 	 * Creates a new Texture editor.
 	 *
-	 * @param Mode - 
-	 * @param InitToolkitHost - 
-	 * @param Texture - 
+	 * @param Mode 
+	 * @param InitToolkitHost 
+	 * @param Texture 
 	 */
 	virtual TSharedRef<ITextureEditorToolkit> CreateTextureEditor( const EToolkitMode::Type Mode, const TSharedPtr< IToolkitHost >& InitToolkitHost, UTexture* Texture ) = 0;
 };
diff --git a/Engine/Source/Editor/TextureEditor/Public/Interfaces/ITextureEditorToolkit.h b/Engine/Source/Editor/TextureEditor/Public/Interfaces/ITextureEditorToolkit.h
--- a/Engine/Source/Editor/TextureEditor/Public/Interfaces/ITextureEditorToolkit.h
+++ b/Engine/Source/Editor/TextureEditor/Public/Interfaces/ITextureEditorToolkit.h
@@ -1,18 +1,13 @@
 // Copyright 1998-2014 Epic Games, Inc. All Rights Reserved.
 
-/*=============================================================================
-	ITextureEditorToolkit.h: Declares the ITextureEditorToolkit interface.
-=============================================================================*/
-
 #pragma once
 
 #include "Toolkits/IToolkitHost.h"
 
 
-/*-----------------------------------------------------------------------------
-   ITextureEditor
------------------------------------------------------------------------------*/
-
+/**
+ * Interface for texture editor tool kits.
+ */
 class ITextureEditorToolkit
 	: public FAssetEditorToolkit
 {
@@ -38,6 +33,11 @@ public:
 	virtual bool GetFitToViewport() const = 0;
 	virtual void SetFitToViewport( const bool bFitToViewport ) = 0;
 	
+public:
+
+	/**
+	 * Toggles the fit-to-viewport mode.
+	 */
 	void ToggleFitToViewport()
 	{
 		const bool bFitToViewport = GetFitToViewport();
diff --git a/Engine/Source/Editor/TextureEditor/Public/TextureEditor.h b/Engine/Source/Editor/TextureEditor/Public/TextureEditor.h
--- a/Engine/Source/Editor/TextureEditor/Public/TextureEditor.h
+++ b/Engine/Source/Editor/TextureEditor/Public/TextureEditor.h
@@ -1,13 +1,9 @@
 // Copyright 1998-2014 Epic Games, Inc. All Rights Reserved.
 
-/*=============================================================================
-	TextureEditor.h: TextureEditor module public header file.
-=============================================================================*/
-
 #pragma once
 
 
-/* Dependencies
+/* Public Dependencies
  *****************************************************************************/
 
 #include "UnrealEd.h"
@@ -16,7 +12,7 @@
 #include "Toolkits/IToolkit.h"
 
 
-/* Interfaces
+/* Public Includes
  *****************************************************************************/
 
 #include "ITextureEditorToolkit.h"
diff --git a/Engine/Source/Editor/UndoHistory/Private/UndoHistoryModule.cpp b/Engine/Source/Editor/UndoHistory/Private/UndoHistoryModule.cpp
--- a/Engine/Source/Editor/UndoHistory/Private/UndoHistoryModule.cpp
+++ b/Engine/Source/Editor/UndoHistory/Private/UndoHistoryModule.cpp
@@ -1,9 +1,5 @@
 // Copyright 1998-2013 Epic Games, Inc. All Rights Reserved.
 
-/*=============================================================================
-	UndoHistoryModule.cpp: Implements the FUndoHistoryModule class.
-=============================================================================*/
-
 #include "UndoHistoryPrivatePCH.h"
 #include "Editor/WorkspaceMenuStructure/Public/WorkspaceMenuStructureModule.h"
 
diff --git a/Engine/Source/Editor/UndoHistory/Private/UndoHistoryPrivatePCH.h b/Engine/Source/Editor/UndoHistory/Private/UndoHistoryPrivatePCH.h
--- a/Engine/Source/Editor/UndoHistory/Private/UndoHistoryPrivatePCH.h
+++ b/Engine/Source/Editor/UndoHistory/Private/UndoHistoryPrivatePCH.h
@@ -1,12 +1,8 @@
 // Copyright 1998-2013 Epic Games, Inc. All Rights Reserved.
 
-/*=============================================================================
-	UndoHistoryPrivatePCH.h: Pre-compiled header file for the UndoHistory module.
-=============================================================================*/
-
 #pragma once
 
-#include "../Public/UndoHistory.h"
+#include "UndoHistory.h"
 
 
 /* Dependencies
diff --git a/Engine/Source/Editor/UndoHistory/Private/Widgets/SUndoHistory.cpp b/Engine/Source/Editor/UndoHistory/Private/Widgets/SUndoHistory.cpp
--- a/Engine/Source/Editor/UndoHistory/Private/Widgets/SUndoHistory.cpp
+++ b/Engine/Source/Editor/UndoHistory/Private/Widgets/SUndoHistory.cpp
@@ -1,9 +1,5 @@
 // Copyright 1998-2013 Epic Games, Inc. All Rights Reserved.
 
-/*=============================================================================
-	SUndoHistory.cpp: Implements the SUndoHistory class.
-=============================================================================*/
-
 #include "UndoHistoryPrivatePCH.h"
 
 
@@ -208,4 +204,5 @@ FText SUndoHistory::HandleUndoSizeTextBlockText( ) const
 	return FText::Format(LOCTEXT("TransactionCountF", "{0} Transactions ({1})"), FText::AsNumber(UndoList.Num()), FText::AsMemory(GEditor->Trans->GetUndoSize()));
 }
 
+
 #undef LOCTEXT_NAMESPACE
\ No newline at end of file
diff --git a/Engine/Source/Editor/UndoHistory/Private/Widgets/SUndoHistory.h b/Engine/Source/Editor/UndoHistory/Private/Widgets/SUndoHistory.h
--- a/Engine/Source/Editor/UndoHistory/Private/Widgets/SUndoHistory.h
+++ b/Engine/Source/Editor/UndoHistory/Private/Widgets/SUndoHistory.h
@@ -1,9 +1,5 @@
 // Copyright 1998-2013 Epic Games, Inc. All Rights Reserved.
 
-/*=============================================================================
-	SUndoHistory.h: Declares the SUndoHistory class.
-=============================================================================*/
-
 #pragma once
 
 
@@ -39,7 +35,7 @@ public:
 	/**
 	 * Construct this widget
 	 *
-	 * @param InArgs - The declaration data for this widget.
+	 * @param InArgs The declaration data for this widget.
 	 */
 	void Construct( const FArguments& InArgs );
 
